Move SimCacheUSD from USDImporter to its own plugin ChaosCacheUSD.

diff --git a/Engine/Plugins/Experimental/ChaosCachingUSD/ChaosCachingUSD.uplugin b/Engine/Plugins/Experimental/ChaosCachingUSD/ChaosCachingUSD.uplugin
new file mode 100644
--- /dev/null
+++ b/Engine/Plugins/Experimental/ChaosCachingUSD/ChaosCachingUSD.uplugin
@@ -0,0 +1,37 @@
+{
+	"FileVersion": 3,
+	"Version": 1,
+	"VersionName": "1.0",
+	"FriendlyName": "Chaos Caching USD",
+	"Description": "Adds support for using USD files for caching Chaos flesh simulations",
+	"Category": "Importers",
+	"CreatedBy": "Epic Games, Inc.",
+	"CreatedByURL": "https://epicgames.com",
+	"DocsURL": "",
+	"MarketplaceURL": "",
+	"SupportURL": "",
+	"EnabledByDefault": false,
+	"CanContainContent": true,
+	"IsBetaVersion": true,
+	"Installed": false,
+	"SupportedTargetPlatforms": [
+		"Win64"
+	],
+	"Modules": [
+		{
+			"Name": "ChaosCachingUSD",
+			"Type": "Editor",
+			"LoadingPhase": "Default"
+		}
+	],
+	"Plugins": [
+		{
+			"Name": "ChaosCaching",
+			"Enabled": true
+		},
+		{
+			"Name": "USDImporter",
+			"Enabled": true
+		}
+	]
+}
diff --git a/Engine/Plugins/Experimental/ChaosCachingUSD/Source/ChaosCachingUSD/ChaosCachingUSD.build.cs b/Engine/Plugins/Experimental/ChaosCachingUSD/Source/ChaosCachingUSD/ChaosCachingUSD.build.cs
new file mode 100644
--- /dev/null
+++ b/Engine/Plugins/Experimental/ChaosCachingUSD/Source/ChaosCachingUSD/ChaosCachingUSD.build.cs
@@ -0,0 +1,41 @@
+// Copyright Epic Games, Inc. All Rights Reserved.
+
+using System.Linq;
+using System.IO;
+using System.Collections.Generic;
+
+namespace UnrealBuildTool.Rules
+{
+	public class ChaosCachingUSD : ModuleRules
+	{
+		public ChaosCachingUSD(ReadOnlyTargetRules Target) : base(Target)
+		{
+			// Does not compile with C++20:
+			// error C4002: too many arguments for function-like macro invocation 'TF_PP_CAT_IMPL'
+			// warning C5103: pasting '"TF_LOG_STACK_TRACE_ON_ERROR"' and '"TF_LOG_STACK_TRACE_ON_WARNING"' does not result in a valid preprocessing token
+			CppStandard = CppStandardVersion.Cpp17;
+
+			bUseRTTI = true;
+
+			PublicDependencyModuleNames.AddRange(
+				new string[]
+				{
+				"Boost",
+				"Core",
+				"CoreUObject",
+				"Engine",
+				"Projects",
+				"RHI",
+				}
+			);
+
+			PrivateDependencyModuleNames.AddRange(
+				new string[] {
+				"UnrealUSDWrapper",
+				"USDClasses",
+				"USDUtilities",
+				}
+			);
+		}
+	}
+}
\ No newline at end of file
diff --git a/Engine/Plugins/Experimental/ChaosCachingUSD/Source/ChaosCachingUSD/Private/ChaosCachingUSD/Module.cpp b/Engine/Plugins/Experimental/ChaosCachingUSD/Source/ChaosCachingUSD/Private/ChaosCachingUSD/Module.cpp
new file mode 100644
--- /dev/null
+++ b/Engine/Plugins/Experimental/ChaosCachingUSD/Source/ChaosCachingUSD/Private/ChaosCachingUSD/Module.cpp
@@ -0,0 +1,39 @@
+// Copyright Epic Games, Inc. All Rights Reserved.
+
+#include "Modules/ModuleInterface.h"
+#include "Modules/ModuleManager.h"
+#include "Interfaces/IPluginManager.h"
+#include "Misc/Paths.h"
+
+#include "UnrealUSDWrapper.h"
+#include "USDMemory.h"
+
+
+class FChaosCachingUSDModule : public IModuleInterface
+{
+	//~ Begin IModuleInterface API
+	virtual void StartupModule() override;
+	//~ End IModuleInterface API
+
+};
+
+void FChaosCachingUSDModule::StartupModule()
+{
+#if USE_USD_SDK
+	// Register the ChaosCachingUSD plugin with USD.
+	IPluginManager& UEPluginManager = IPluginManager::Get();
+	FString USDImporterDir = UEPluginManager.FindPlugin(TEXT("USDImporter"))->GetBaseDir();
+
+	const FString ChaosCachingUSDResourcesDir =
+		FPaths::ConvertRelativePathToFull(
+			FPaths::Combine(
+				USDImporterDir, 
+				FString(TEXT("ChaosCachingUSD")),
+				FString(TEXT("Resources"))));
+
+	UnrealUSDWrapper::RegisterPlugins(ChaosCachingUSDResourcesDir);
+#endif // USE_USD_SDK
+}
+
+
+IMPLEMENT_MODULE_USD(FChaosCachingUSDModule, ChaosCachingUSD);
diff --git a/Engine/Plugins/Experimental/ChaosCachingUSD/Source/ChaosCachingUSD/Private/ChaosCachingUSD/Operations.cpp b/Engine/Plugins/Experimental/ChaosCachingUSD/Source/ChaosCachingUSD/Private/ChaosCachingUSD/Operations.cpp
new file mode 100644
--- /dev/null
+++ b/Engine/Plugins/Experimental/ChaosCachingUSD/Source/ChaosCachingUSD/Private/ChaosCachingUSD/Operations.cpp
@@ -0,0 +1,677 @@
+// Copyright Epic Games, Inc. All Rights Reserved.
+
+#include "ChaosCachingUSD/Operations.h"
+#include "ChaosCachingUSD/UEUsdGeomTetMesh.h"
+#include "ChaosCachingUSD/Util.h"
+
+#if USE_USD_SDK
+
+#include "Containers/StringConv.h"
+#include "GeometryCollection/GeometryCollection.h"
+#include "GeometryCollection/ManagedArray.h"
+#include "GeometryCollection/ManagedArrayCollection.h"
+#include "Misc/Paths.h"
+
+#include "UnrealUSDWrapper.h"
+#include "USDConversionUtils.h"
+#include "USDLog.h"
+#include "USDMemory.h"
+
+#include "UsdWrappers/SdfLayer.h"
+#include "UsdWrappers/SdfPath.h"
+#include "UsdWrappers/UsdAttribute.h"
+#include "UsdWrappers/UsdPrim.h"
+#include "UsdWrappers/UsdStage.h"
+
+#include "USDIncludesStart.h"
+	#include "pxr/usd/sdf/primSpec.h"
+	#include "pxr/usd/usd/stageCacheContext.h"
+	#include "pxr/usd/usdGeom/mesh.h"
+	#include "pxr/usd/usdGeom/pointBased.h"
+	#include "pxr/usd/usdGeom/tokens.h"
+	#include "pxr/usd/usdUtils/stageCache.h"
+#include "USDIncludesEnd.h"
+
+bool 
+UE::ChaosCachingUSD::NewStage(const FString& StageName, UE::FUsdStage& UsdStage)
+{
+	if (FPaths::FileExists(StageName))
+	{
+		UE_LOG(LogUsd, Error, TEXT("Cannot clobber existing file: '%s'"), *StageName);
+		return false;
+	}
+
+	// USD caches all stages you open/create, unless you tell it not to.
+	// For convenience, we're using the global/default stage cache.
+	UsdStage = UnrealUSDWrapper::NewStage(*StageName);
+	if (!UsdStage)
+	{
+		UE_LOG(LogUsd, Error, TEXT("Failed to create new USD stage: '%s'"), *StageName);
+		return false;
+	}
+
+	// TODO: Ryan
+	//UsdUtils::SetUsdStageMetersPerUnit(UsdStage, ...);
+	UsdUtils::SetUsdStageUpAxis(UsdStage, pxr::UsdGeomTokens->z);
+	return true;
+}
+
+bool 
+UE::ChaosCachingUSD::OpenStage(const FString& StageName, UE::FUsdStage& UsdStage)
+{
+	if (!FPaths::FileExists(StageName))
+	{
+		UE_LOG(LogUsd, Error, TEXT("File not found: '%s'"), *StageName);
+		return false;
+	}
+
+	// USD caches all stages you open/create, unless you tell it not to.
+	// For convenience, we're using the global/default stage cache.
+	UsdStage = UnrealUSDWrapper::OpenStage(
+			*StageName,						
+			EUsdInitialLoadSet::LoadAll,	// load payload?
+			true,							// use stage cache
+			false);							// force reload from disk
+	if (!UsdStage)
+	{
+		UE_LOG(LogUsd, Error, TEXT("Failed to open USD stage: '%s'"), *StageName);
+		return false;
+	}
+	return true;
+}
+
+bool 
+UE::ChaosCachingUSD::SaveStage(UE::FUsdStage& Stage, const double FirstFrame, const double LastFrame)
+{
+	if (FirstFrame != -TNumericLimits<double>::Max() && LastFrame != -TNumericLimits<double>::Max())
+	{
+		const double StageStart = Stage.GetStartTimeCode();
+		const double StageEnd = Stage.GetEndTimeCode();
+		const bool bHasRange = static_cast<pxr::UsdStageRefPtr>(Stage)->HasAuthoredTimeCodeRange();
+		if (!bHasRange || FirstFrame < StageStart)
+		{
+			Stage.SetStartTimeCode(FirstFrame);
+		}
+		if (!bHasRange || LastFrame > StageEnd)
+		{
+			Stage.SetEndTimeCode(LastFrame);
+		}
+	}
+	return Stage.GetRootLayer().Save();
+}
+
+bool
+UE::ChaosCachingUSD::CloseStage(const UE::FUsdStage& Stage)
+{
+	UnrealUSDWrapper::EraseStageFromCache(Stage);
+	return true;
+}
+
+bool 
+UE::ChaosCachingUSD::CloseStage(const FString& StageName)
+{
+	TArray<UE::FUsdStage> CachedStages = UnrealUSDWrapper::GetAllStagesFromCache();
+	for (UE::FUsdStage& Stage : CachedStages)
+	{
+		if (Stage.GetRootLayer().GetRealPath() == StageName)
+		{
+			return ChaosCachingUSD::CloseStage(Stage);
+		}
+	}
+	return true;
+}
+
+void
+UE::ChaosCachingUSD::GenerateValueClipStageNames(
+	const FString& ParentName, 
+	FString& TopologyName, 
+	FString& TimeVaryingTemplate)
+{
+	FScopedUnrealAllocs UEAllocs; // Use UE memory allocator
+	std::string ParentNameStr(TCHAR_TO_ANSI(*ParentName));
+	std::string TopologyNameStr = pxr::UsdUtilsGenerateClipTopologyName(ParentNameStr);
+	TopologyName = FString(ANSI_TO_TCHAR(TopologyNameStr.c_str()));
+	// Usd supports an integral template format, and a floating point.  We're using the
+	// float format here.  The integral format is "path/name.####.usd".
+	TimeVaryingTemplate = TopologyName.Replace(TEXT("topology"), TEXT("###.###"));
+}
+
+FString
+UE::ChaosCachingUSD::GenerateValueClipTimeVaryingStageName(
+	const FString& TimeVaryingTemplate, 
+	const double Time)
+{
+	FScopedUnrealAllocs UEAllocs; // Use UE memory allocator
+	if (Time == -TNumericLimits<double>::Max())
+	{
+		return TimeVaryingTemplate.Replace(TEXT("###.###"), TEXT("default"));
+	}
+	FString Template = TimeVaryingTemplate.Replace(TEXT("###.###"), TEXT("{0}"));
+	FString TimeStr = FString::Printf(TEXT("%06.3f"), Time);
+	return FString::Format(*Template, { TimeStr });
+}
+
+bool
+UE::ChaosCachingUSD::NewValueClipsStages(
+	const FString& ParentStageName,
+	const FString& TopologyStageName,
+	UE::FUsdStage& ParentStage,
+	UE::FUsdStage& TopologyStage)
+{
+	if (!ChaosCachingUSD::NewStage(ParentStageName, ParentStage)) return false;
+	if (!ChaosCachingUSD::NewStage(TopologyStageName, TopologyStage)) return false;
+	return true;
+}
+
+bool
+UE::ChaosCachingUSD::NewValueClipsFrameStage(
+	const FString& TimeVaryingStageTemplate,
+	const double Time,
+	FString& FrameStageName,
+	UE::FUsdStage& FrameStage)
+{
+	FrameStageName =
+		ChaosCachingUSD::GenerateValueClipTimeVaryingStageName(TimeVaryingStageTemplate, Time);
+	return ChaosCachingUSD::NewStage(FrameStageName, FrameStage);
+}
+
+bool 
+UE::ChaosCachingUSD::InitValueClipsTemplate(
+	UE::FUsdStage& ParentStage,
+	UE::FUsdStage& TopologyStage,
+	const FString& ParentStageName,
+	const FString& TopologyStageName,
+	const FString& TimeVaryingStageTemplate,
+	const TArray<FString>& PrimPaths,
+	const double StartTime,
+	const double EndTime,
+	const double Stride)
+{
+	FScopedUsdAllocs UEAllocs; // Use USD memory allocator
+
+	UE::FSdfLayer ParentLayer = ParentStage.GetRootLayer();
+	UE::FSdfLayer TopologyLayer = ParentStage.GetRootLayer();
+
+	for (const FString& PrimPath : PrimPaths)
+	{
+		// Get or create the prim.
+		pxr::SdfPrimSpecHandle Prim = pxr::SdfCreatePrimInLayer(ParentLayer, UE::FSdfPath(*PrimPath));
+
+		// Get the existing metadata, if it exists.
+		pxr::VtDictionary Clips;
+		pxr::VtDictionary ClipSetDict;
+		if (Prim->HasInfo(pxr::UsdTokens->clips))
+		{
+			pxr::VtValue PrevClipInfoValue = Prim->GetInfo(pxr::UsdTokens->clips);
+			if (PrevClipInfoValue.IsHolding<pxr::VtDictionary>())
+			{
+				Clips = PrevClipInfoValue.Get<pxr::VtDictionary>();
+				auto it = Clips.find(pxr::UsdClipsAPISetNames->default_);
+				if (it != Clips.end())
+				{
+					pxr::VtValue PrevClipsEntryValue = it->second;
+					if (PrevClipsEntryValue.IsHolding<pxr::VtDictionary>())
+					{
+						ClipSetDict = PrevClipsEntryValue.Get<pxr::VtDictionary>();
+					}
+				}
+			}
+		}
+
+		// Add the topology layer.
+		if (!ParentStageName.IsEmpty() && !TopologyStageName.IsEmpty() && TopologyLayer)
+		{
+			FString TopologyId = TopologyStageName;
+			FPaths::MakePathRelativeTo(TopologyId, *ParentStageName);
+			std::string TopologyIdStr(TCHAR_TO_ANSI(*TopologyId));
+
+			// Add the sub layer in the parent pointing to the topology.
+			static_cast<pxr::SdfLayerRefPtr>(ParentLayer)->InsertSubLayerPath(
+				TopologyIdStr, 0); // 0 means strongest
+			
+			// Set root layer metadata
+			ClipSetDict[pxr::UsdClipsAPIInfoKeys->manifestAssetPath] = 
+				pxr::SdfAssetPath(TopologyIdStr);
+		}
+
+		// Add time varying stage template.
+		if (!TimeVaryingStageTemplate.IsEmpty())
+		{
+			std::string TimeVaryingStageTemplateStr(TCHAR_TO_ANSI(*TimeVaryingStageTemplate));
+			ClipSetDict[pxr::UsdClipsAPIInfoKeys->templateAssetPath] = 
+				TimeVaryingStageTemplateStr;
+		}
+
+		// Add the prim path.
+		std::string PrimPathStr(TCHAR_TO_ANSI(*PrimPath));
+		ClipSetDict[pxr::UsdClipsAPIInfoKeys->primPath] = PrimPathStr;
+
+		// Add start and end time.
+		if (fabs(StartTime) != TNumericLimits<double>::Max() &&
+			fabs(EndTime) != TNumericLimits<double>::Max())
+		{
+			ClipSetDict[pxr::UsdClipsAPIInfoKeys->templateStartTime] = StartTime;
+			ClipSetDict[pxr::UsdClipsAPIInfoKeys->templateEndTime] = EndTime;
+		}
+
+		// Add sampling stride.
+		if (Stride > 0.0)
+		{
+			ClipSetDict[pxr::UsdClipsAPIInfoKeys->templateStride] = Stride;
+		}
+
+		// Set metadata info on Prim.
+		Clips[pxr::UsdClipsAPISetNames->default_] = ClipSetDict;
+		Prim->SetInfo(pxr::UsdTokens->clips, pxr::VtValue::Take(Clips));
+	}
+	ParentLayer.SetStartTimeCode(StartTime);
+	ParentLayer.SetEndTimeCode(EndTime);
+	return true;
+}
+
+bool
+UE::ChaosCachingUSD::InitValueClipsTemplate(
+	const FString& ParentStageName,
+	const FString& TopologyStageName,
+	const FString& TimeVaryingStageTemplate,
+	const TArray<FString>& PrimPaths,
+	const double StartTime,
+	const double EndTime,
+	const double Stride)
+{
+	UE::FUsdStage ParentStage;
+	if (!ChaosCachingUSD::OpenStage(ParentStageName, ParentStage)) return false;
+	UE::FUsdStage TopologyStage;
+	if (!ChaosCachingUSD::OpenStage(TopologyStageName, TopologyStage)) return false;
+
+	return ChaosCachingUSD::InitValueClipsTemplate(
+		ParentStage,
+		TopologyStage,
+		ParentStageName,
+		TopologyStageName,
+		TimeVaryingStageTemplate,
+		PrimPaths,
+		StartTime,
+		EndTime,
+		Stride);
+}
+
+bool 
+UE::ChaosCachingUSD::WriteTetMesh(
+	UE::FUsdStage& Stage, 
+	const FString& PrimPath,
+	const FManagedArrayCollection& Collection, 
+	const int32 StructureIndex)
+{
+	FScopedUsdAllocs UEAllocs; // Use USD memory allocator
+
+	const TManagedArray<FIntVector4>* Tetrahedron =
+		Collection.FindAttribute<FIntVector4>(
+			/*FTetrahedralCollection::TetrahedronAttribute*/FName("Tetrahedron"), 
+			/*FTetrahedralCollection::TetrahedralGroup*/FName("Tetrahedral"));
+	const TManagedArray<int32>* TetrahedronStart =
+		Collection.FindAttribute<int32>(
+			/*FTetrahedralCollection::TetrahedronStartAttribute*/FName("TetrahedronStart"), 
+			/*FGeometryCollection::GeometryGroup*/FName("Tetrahedral"));
+	const TManagedArray<int32>* TetrahedronCount =
+		Collection.FindAttribute<int32>(
+			/*FTetrahedralCollection::TetrahedronCountAttribute*/FName("TetrahedronCount"), 
+			/*FGeometryCollection::GeometryGroup*/FName("Geometry"));
+
+	const TManagedArray<FIntVector>* Triangle =
+		Collection.FindAttribute<FIntVector>(
+			"Indices", /*FGeometryCollection::FacesGroup*/FName("Faces"));
+	const TManagedArray<int32>* FacesStart =
+		Collection.FindAttribute<int32>(
+			"FaceStart", /*FGeometryCollection::GeometryGroup*/FName("Geometry"));
+	const TManagedArray<int32>* FacesCount =
+		Collection.FindAttribute<int32>(
+			"FaceCount", /*FGeometryCollection::GeometryGroup*/FName("Geometry"));
+
+	const TManagedArray<int32>* VertexStart =
+		Collection.FindAttribute<int32>(
+			"VertexStart", /*FGeometryCollection::GeometryGroup*/FName("Geometry"));
+	const TManagedArray<int32>* VertexCount =
+		Collection.FindAttribute<int32>(
+			"VertexCount", /*FGeometryCollection::GeometryGroup*/FName("Geometry"));
+
+	const TManagedArray<FVector3f>* Vertex =
+		Collection.FindAttribute<FVector3f>(
+			"Vertex", "Vertices");
+
+	int32 VertStart = 0;
+	int32 VertCount = Vertex->Num();
+	int32 TetMeshStart = 0;
+	int32 TetMeshCount = Tetrahedron->Num();
+	int32 TriMeshStart = 0;
+	int32 TriMeshCount = Triangle->Num();
+	if (StructureIndex != INDEX_NONE)
+	{
+		if (!VertexStart->IsValidIndex(StructureIndex) ||
+			!VertexCount->IsValidIndex(StructureIndex))
+		{
+			UE_LOG(LogUsd, Error, TEXT("Invalid vertex index: %d, valid range [0, %d)."), StructureIndex, VertexStart->Num());
+			return false;
+		}
+		VertStart = (*VertexStart)[StructureIndex];
+		VertCount = (*VertexCount)[StructureIndex];
+		if (!TetrahedronStart->IsValidIndex(StructureIndex) ||
+			!TetrahedronCount->IsValidIndex(StructureIndex))
+		{
+			UE_LOG(LogUsd, Error, TEXT("Invalid tet mesh index: %d, valid range [0, %d)."), StructureIndex, TetrahedronStart->Num());
+			return false;
+		}
+		TetMeshStart = (*TetrahedronStart)[StructureIndex];
+		TetMeshCount = (*TetrahedronCount)[StructureIndex];
+		if (!FacesStart->IsValidIndex(StructureIndex) ||
+			!FacesCount->IsValidIndex(StructureIndex))
+		{
+			// Not a deal breaker, but highly unusual.
+			UE_LOG(LogUsd, Warning, TEXT("Invalid tet mesh faces index: %d, valid range [0, %d),"), StructureIndex, FacesStart->Num());
+		}
+		else
+		{
+			TriMeshStart = (*FacesStart)[StructureIndex];
+			TriMeshCount = (*FacesCount)[StructureIndex];
+		}
+
+		UE_LOG(LogUsd, Log, TEXT("Extracting tetrahedral mesh %d of %d from collection; vertex range [%d, %d), tet [%d, %d), tri [%d, %d)."),
+			StructureIndex, VertexStart->Num(),
+			VertStart, VertStart + VertCount,
+			TetMeshStart, TetMeshStart + TetMeshCount,
+			TriMeshStart, TriMeshStart + TriMeshCount);
+	}
+
+	UE::FSdfPath Path(*PrimPath);
+
+	// Author a UsdGeomPointBased override prim.
+	if (!UE::ChaosCachingUSD::WritePoints(Stage, PrimPath, -TNumericLimits<double>::Max(), Collection, StructureIndex))
+	{
+		UE_LOG(LogUsd, Error, TEXT("Failed to create UsdGeomPointBased at path: '%s' on USD stage: '%s'."),
+			*PrimPath, *Stage.GetRootLayer().GetRealPath());
+		return false;
+	}
+
+	// Promote to UsdGeomMesh.
+	if (pxr::UsdGeomMesh Mesh = pxr::UsdGeomMesh::Define(Stage, Path))
+	{
+		pxr::UsdPrim Prim = Mesh.GetPrim();
+
+		// Surface tris vertices and counts
+		pxr::VtArray<int> VtTriVerts(TriMeshCount * 3);
+		int32 k = 0;
+		for (int32 i = 0; i < TriMeshCount; i++)
+			for (int32 j = 0; j < 3; j++, k++)
+				VtTriVerts[k] = (*Triangle)[TriMeshStart + i][j] - VertStart;
+		Mesh.CreateFaceVertexIndicesAttr().Set(VtTriVerts);
+		pxr::VtArray<int> VtTriCounts(TriMeshCount);
+		for (int32 i = 0; i < TriMeshCount; i++)
+			VtTriCounts[i] = 3; // (This makes me die inside)
+		Mesh.CreateFaceVertexCountsAttr().Set(VtTriCounts);
+
+		// Subdivision scheme set to 'none' in UEUsdGeomTetMesh constructor, but we do it here anyway.
+		Mesh.CreateSubdivisionSchemeAttr().Set(pxr::UsdGeomTokens->none);
+	}
+	else
+	{
+		UE_LOG(LogUsd, Warning, TEXT("Failed to create UsdGeomMesh at path: '%s' on USD stage: '%s'."),
+			*PrimPath, *Stage.GetRootLayer().GetRealPath());
+	}
+
+	// Promote to UEUsdGeomTetMesh.
+	// This next line succeeds, but farts out a:
+	//   TF_DIAGNOSTIC_CODING_ERROR_TYPE: Failed to find plugin for schema type 'UEUsdGeomTetMesh'
+	// I think that means there's an issue with the plugInfo.json, and I think that means that if
+	// USD was to happen upon a prim of this type, it might not know which plugin to load.  For 
+	// our caching purposes, I think that's of no consequence; other than the fact that operator bool 
+	// returns false for TetMesh, which I think may be related, but may not.  PointBased's operator 
+	// bool does the same thing.  What's more is that we currently have no consumer for the tet 
+	// topology, so this really doesn't matter.  At least not yet.
+	pxr::UEUsdGeomTetMesh TetMesh = pxr::UEUsdGeomTetMesh::Define(Stage, Path);
+	if (/*TetMesh*/true)
+	{
+		pxr::UsdPrim Prim = TetMesh.GetPrim();
+
+		// Tetrahedra
+		check(sizeof(pxr::GfVec4i) == sizeof(FIntVector4));
+		const pxr::GfVec4i* GfTets = reinterpret_cast<const pxr::GfVec4i*>(&(*Tetrahedron)[TetMeshStart]);
+		pxr::VtArray<pxr::GfVec4i> VtTets(TetMeshCount);
+		for (int32 i = 0; i < TetMeshCount; i++) 
+			VtTets[i] = GfTets[i] - pxr::GfVec4i(VertStart);
+		TetMesh.CreateTetVertexIndicesAttr().Set(VtTets);
+	}
+	else
+	{
+		UE_LOG(LogUsd, Warning, TEXT("Failed to create UEUsdGeomTetMesh at path: '%s' on USD stage: '%s'."),
+			*PrimPath, *Stage.GetRootLayer().GetRealPath());
+	}
+
+	UE::ChaosCachingUSD::SaveStage(Stage, -TNumericLimits<double>::Max(), -TNumericLimits<double>::Max());
+	return true;
+}
+
+bool
+UE::ChaosCachingUSD::WritePoints(
+	UE::FUsdStage& Stage,
+	const FString& PrimPath,
+	const double Time,
+	pxr::VtArray<pxr::GfVec3f>& VtPoints,
+	pxr::VtArray<pxr::GfVec3f>& VtVels)
+{
+	pxr::GfMatrix4d GfMat(1.0);
+
+	UE::FSdfPath Path(*PrimPath);
+	pxr::UsdPrim Prim = Stage.GetPrimAtPath(Path);
+	//UE::FUsdPrim Prim = PrimPath.IsEmpty() ? Stage.GetDefaultPrim() : Stage.GetPrimAtPath(UE::FSdfPath(*PrimPath));
+	if (!Prim)
+	{
+		// We're authoring a prim that doesn't exist.  Use an override prim.
+		Prim = Stage.OverridePrim(Path);
+		//pxr::UsdGeomPointBased PointBased(Prim);
+		//PointBased.CreatePointsAttr();
+		//PointBased.CreateExtentAttr();
+
+		UE::ChaosCachingUSDUtil::Private::DefineAncestorTransforms(Stage, Path);
+
+		return UE::ChaosCachingUSDUtil::Private::SetPointsExtentAndPrimXform(
+			Prim, GfMat, VtPoints, VtVels,
+			Time == -TNumericLimits<double>::Max() ? pxr::UsdTimeCode::Default() : pxr::UsdTimeCode(Time));
+	}
+	else
+	{
+		// We're probably authoring an override prim, UsdGeomPointBased, or a UsdGeomPointInstancer.
+		return UE::ChaosCachingUSDUtil::Private::SetPointsExtentAndPrimXform(
+			Prim, GfMat, VtPoints, VtVels,
+			Time == -TNumericLimits<double>::Max() ? pxr::UsdTimeCode::Default() : pxr::UsdTimeCode(Time));
+	}
+}
+
+bool 
+UE::ChaosCachingUSD::WritePoints(
+	UE::FUsdStage& Stage, 
+	const FString& PrimPath, 
+	const double Time, 
+	const FManagedArrayCollection& Collection, 
+	const int32 StructureIndex)
+{
+	FScopedUsdAllocs UEAllocs; // Use USD memory allocator
+
+	const TManagedArray<int32>* VertexStart =
+		Collection.FindAttribute<int32>(
+			"VertexStart", /*FGeometryCollection::GeometryGroup*/ FName("Geometry"));
+	const TManagedArray<int32>* VertexCount =
+		Collection.FindAttribute<int32>(
+			"VertexCount", /*FGeometryCollection::GeometryGroup*/FName("Geometry"));
+
+	const TManagedArray<FVector3f>* Vertex =
+		Collection.FindAttribute<FVector3f>(
+			"Vertex", "Vertices");
+
+	int32 VertStart = 0;
+	int32 VertCount = Vertex->Num();
+	if (StructureIndex != INDEX_NONE)
+	{
+		if (!VertexStart->IsValidIndex(StructureIndex) ||
+			!VertexCount->IsValidIndex(StructureIndex))
+		{
+			UE_LOG(LogUsd, Error, TEXT("Invalid vertex index: %d, valid range [0, %d)."), StructureIndex, VertexStart->Num());
+			return false;
+		}
+		VertStart = (*VertexStart)[StructureIndex];
+		VertCount = (*VertexCount)[StructureIndex];
+	}
+
+	check(sizeof(pxr::GfVec3f) == sizeof(FVector3f)); // reinterpret_cast works?
+	pxr::VtArray<pxr::GfVec3f> VtPoints(VertCount); // VtArray is shared mem with copy-on-write.
+	for (int32 i = 0; i < VertCount; i++)
+		VtPoints[i] = *reinterpret_cast<const pxr::GfVec3f*>(&(*Vertex)[VertStart + i]);
+	pxr::GfMatrix4d GfMat(1.0);
+
+	UE::FSdfPath Path(*PrimPath);
+	pxr::UsdPrim Prim = Stage.GetPrimAtPath(Path);
+	if(!Prim)
+	{
+		// We're authoring a prim that doesn't exist.  Use an override prim.
+		Prim = Stage.OverridePrim(Path);
+		//pxr::UsdGeomPointBased PointBased(Prim);
+		//PointBased.CreatePointsAttr();
+		//PointBased.CreateExtentAttr();
+
+		UE::ChaosCachingUSDUtil::Private::DefineAncestorTransforms(Stage, Path);
+
+		return UE::ChaosCachingUSDUtil::Private::SetPointsExtentAndPrimXform(
+			Prim, GfMat, VtPoints,
+			Time == -TNumericLimits<double>::Max() ? pxr::UsdTimeCode::Default() : pxr::UsdTimeCode(Time));
+	}
+	else 
+	{
+		// We're probably authoring an override prim, UsdGeomPointBased, or a UsdGeomPointInstancer.
+		return UE::ChaosCachingUSDUtil::Private::SetPointsExtentAndPrimXform(
+			Prim, GfMat, VtPoints,
+			Time == -TNumericLimits<double>::Max() ? pxr::UsdTimeCode::Default() : pxr::UsdTimeCode(Time));
+	}
+}
+
+bool 
+UE::ChaosCachingUSD::ReadTimeSamples(
+	const UE::FUsdStage& Stage, 
+	const FString& PrimPath, 
+	const FString& AttrName, 
+	TArray<double>& TimeSamples)
+{
+	UE::FSdfPath Path(*PrimPath);
+	UE::FUsdPrim Prim = Stage.GetPrimAtPath(Path);
+	if (!Prim)
+	{
+		UE_LOG(LogUsd, Error, TEXT("No prim found at path '%s'."), *PrimPath);
+		return false;
+	}
+	UE::FUsdAttribute Attr = Prim.GetAttribute(*AttrName);
+	if (!Attr)
+	{
+		UE_LOG(LogUsd, Error, TEXT("No attribute '%s' found on prim '%s'."), *AttrName, *PrimPath);
+		return false;
+	}
+	Attr.GetTimeSamples(TimeSamples);
+	return true;
+}
+
+bool 
+UE::ChaosCachingUSD::ReadTimeSamples(
+	const UE::FUsdStage& Stage, 
+	const FString& PrimPath, 
+	TArray<double>& TimeSamples)
+{
+	return ChaosCachingUSD::ReadTimeSamples(Stage, PrimPath, FString(pxr::UsdGeomTokens->points.GetText()), TimeSamples);
+}
+
+FString
+UE::ChaosCachingUSD::GetPointsAttrName()
+{
+	FScopedUnrealAllocs UEAllocs; // Use UE memory allocator
+	return FString(pxr::UsdGeomTokens->points.GetText());
+}
+
+FString
+UE::ChaosCachingUSD::GetVelocityAttrName()
+{
+	FScopedUnrealAllocs UEAllocs; // Use UE memory allocator
+	return FString(pxr::UsdGeomTokens->velocities.GetText());
+}
+
+bool 
+UE::ChaosCachingUSD::GetBracketingTimeSamples(
+	const UE::FUsdStage& Stage, 
+	const FString& PrimPath, 
+	const FString& AttrName, 
+	const double TargetTime, 
+	double* Lower, 
+	double* Upper)
+{
+	FScopedUsdAllocs UEAllocs; // Use USD memory allocator
+
+	UE::FSdfPath Path(*PrimPath);
+	UE::FUsdPrim Prim = Stage.GetPrimAtPath(Path);
+	if (!Prim)
+	{
+		UE_LOG(LogUsd, Error, TEXT("No prim found at path '%s'."), *PrimPath);
+		return false;
+	}
+	UE::FUsdAttribute Attr = Prim.GetAttribute(*AttrName);
+	if (!Attr)
+	{
+		UE_LOG(LogUsd, Error, TEXT("No attribute '%s' found on prim '%s'."), *AttrName, *PrimPath);
+		return false;
+	}
+
+	bool bHasTimeSamples = false;
+	return static_cast<pxr::UsdAttribute&>(Attr).GetBracketingTimeSamples(
+		TargetTime, Lower, Upper, &bHasTimeSamples) && bHasTimeSamples;
+}
+
+
+bool
+UE::ChaosCachingUSD::ReadPoints(
+	const UE::FUsdStage& Stage,
+	const FString& PrimPath,
+	const FString& AttrName,
+	const double Time,
+	pxr::VtArray<pxr::GfVec3f>& VtPoints)
+{
+	FScopedUsdAllocs UEAllocs; // Use USD memory allocator
+	
+	UE::FSdfPath Path(*PrimPath);
+	UE::FUsdPrim Prim = Stage.GetPrimAtPath(Path);
+	if (!Prim)
+	{
+		UE_LOG(LogUsd, Error, TEXT("No prim found at path '%s'."), *PrimPath);
+		return false;
+	}
+	UE::FUsdAttribute Attr = Prim.GetAttribute(*AttrName);
+	if (!Attr)
+	{
+		UE_LOG(LogUsd, Error, TEXT("No attribute '%s' found on prim '%s'."), *AttrName, *PrimPath);
+		return false;
+	}
+
+	// VtArray is copy-on-write. Be careful to use const API.
+	static_cast<pxr::UsdAttribute&>(Attr).Get(
+		&VtPoints, 
+		Time == -TNumericLimits<double>::Max() ? pxr::UsdTimeCode::Default() : pxr::UsdTimeCode(Time));
+
+	return true;
+}
+
+bool
+UE::ChaosCachingUSD::ReadPoints(
+	const UE::FUsdStage& Stage,
+	const FString& PrimPath,
+	const double Time,
+	pxr::VtArray<pxr::GfVec3f>& VtPoints,
+	pxr::VtArray<pxr::GfVec3f>& VtVels)
+{
+	ChaosCachingUSD::ReadPoints(Stage, PrimPath, FString(pxr::UsdGeomTokens->velocities.GetText()), Time, VtVels);
+	return ChaosCachingUSD::ReadPoints(Stage, PrimPath, FString(pxr::UsdGeomTokens->points.GetText()), Time, VtPoints);
+}
+
+#endif // USE_USD_SDK
diff --git a/Engine/Plugins/Experimental/ChaosCachingUSD/Source/ChaosCachingUSD/Private/ChaosCachingUSD/UEUsdGeomTetMesh.cpp b/Engine/Plugins/Experimental/ChaosCachingUSD/Source/ChaosCachingUSD/Private/ChaosCachingUSD/UEUsdGeomTetMesh.cpp
new file mode 100644
--- /dev/null
+++ b/Engine/Plugins/Experimental/ChaosCachingUSD/Source/ChaosCachingUSD/Private/ChaosCachingUSD/UEUsdGeomTetMesh.cpp
@@ -0,0 +1,172 @@
+// Copyright Epic Games, Inc. All Rights Reserved.
+
+#include "ChaosCachingUSD/UEUsdGeomTetMesh.h"
+#include "ChaosCachingUSD/UEUsdGeomTokens.h"
+
+#if USE_USD_SDK
+
+#include "USDMemory.h"
+#include "USDIncludesStart.h"
+	#include "pxr/usd/sdf/assetPath.h"
+	#include "pxr/usd/sdf/types.h"
+	#include "pxr/usd/usd/schemaRegistry.h"
+	#include "pxr/usd/usd/typed.h"
+
+PXR_NAMESPACE_OPEN_SCOPE
+// Register the schema with the TfType system - need to do this inside 
+// USDIncludesStart/End.h so that warnings are disabled, and re-enabled 
+// in USDIncludesEnd.h.
+#pragma warning(disable: 4191) /* 'reinterpret_cast': unsafe conversion during TF_REGISTRY_FUNCTION(TfType) */
+TF_REGISTRY_FUNCTION(TfType)
+{
+	FScopedUsdAllocs UEAllocs; // Use USD memory allocator
+	TfType::Define<UEUsdGeomTetMesh, TfType::Bases<UsdGeomMesh> >();
+	TfType::AddAlias<UsdSchemaBase, UEUsdGeomTetMesh>("UETetMesh");
+}
+PXR_NAMESPACE_CLOSE_SCOPE
+
+#include "USDIncludesEnd.h"
+
+PXR_NAMESPACE_OPEN_SCOPE
+
+UEUsdGeomTetMesh::UEUsdGeomTetMesh(const UsdPrim& Prim)
+    : UsdGeomMesh(Prim)          
+{
+    if(pxr::UsdAttribute SubdivisionAttr = UsdGeomMesh::CreateSubdivisionSchemeAttr())
+    {
+        SubdivisionAttr.Set(pxr::UsdGeomTokens->none);
+    }
+}
+
+UEUsdGeomTetMesh::UEUsdGeomTetMesh(const UsdSchemaBase& SchemaObj)
+    : UsdGeomMesh(SchemaObj)          
+{
+    if(pxr::UsdAttribute SubdivisionAttr = UsdGeomMesh::CreateSubdivisionSchemeAttr())
+    {
+        SubdivisionAttr.Set(pxr::UsdGeomTokens->none);
+    }
+}
+
+UEUsdGeomTetMesh::~UEUsdGeomTetMesh()
+{}
+
+UEUsdGeomTetMesh
+UEUsdGeomTetMesh::Get(const UsdStagePtr &Stage, const SdfPath &Path)
+{
+	if (!Stage)
+	{
+		TF_CODING_ERROR("Invalid stage");
+		return UEUsdGeomTetMesh();
+	}
+	return UEUsdGeomTetMesh(Stage->GetPrimAtPath(Path));
+}
+
+UEUsdGeomTetMesh
+UEUsdGeomTetMesh::Define(const UsdStagePtr &Stage, const SdfPath &Path)
+{
+    static TfToken UsdPrimTypeName("UETetMesh");
+    if(!Stage)
+    {
+        TF_CODING_ERROR("Invalid stage");
+        return UEUsdGeomTetMesh();
+    }
+    return UEUsdGeomTetMesh(Stage->DefinePrim(Path, UsdPrimTypeName));
+}
+
+UsdSchemaKind
+UEUsdGeomTetMesh::_GetSchemaKind() const
+{
+    return UEUsdGeomTetMesh::schemaKind;
+}
+
+const TfType &
+UEUsdGeomTetMesh::_GetStaticTfType()
+{
+    static TfType Type = TfType::Find<UEUsdGeomTetMesh>();
+    return Type;
+}
+
+bool
+UEUsdGeomTetMesh::_IsTypedSchema()
+{
+    static bool bIsTyped = _GetStaticTfType().IsA<UsdTyped>();
+    return bIsTyped;
+}
+
+const TfType &
+UEUsdGeomTetMesh::_GetTfType() const
+{
+    return _GetStaticTfType();
+}
+
+UsdAttribute
+UEUsdGeomTetMesh::GetTetVertexIndicesAttr() const
+{
+    return GetPrim().GetAttribute(UEUsdGeomTokens->tetVertexIndices);
+}
+
+UsdAttribute
+UEUsdGeomTetMesh::CreateTetVertexIndicesAttr(VtValue const &DefaultValue, bool bWriteSparsely) const
+{
+    return UsdSchemaBase::_CreateAttr(
+		UEUsdGeomTokens->tetVertexIndices,
+		SdfValueTypeNames->Int4Array,
+		false,    // custom
+		SdfVariabilityVarying,
+		DefaultValue,
+		bWriteSparsely);
+}
+
+UsdAttribute 
+UEUsdGeomTetMesh::GetTetOrientationAttr() const
+{
+	return GetPrim().GetAttribute(UEUsdGeomTokens->tetOrientation);
+}
+
+UsdAttribute 
+UEUsdGeomTetMesh::CreateTetOrientationAttr(VtValue const& DefaultValue, bool bWriteSparsely) const
+{
+	return UsdSchemaBase::_CreateAttr(
+		UEUsdGeomTokens->tetOrientation,
+		SdfValueTypeNames->Token,
+		false,	// custom
+		SdfVariabilityUniform,
+		DefaultValue,
+		bWriteSparsely);
+}
+
+namespace {
+	static inline TfTokenVector
+	_ConcatenateAttributeNames(const TfTokenVector& Left, const TfTokenVector& Right)
+	{
+	    TfTokenVector Result;
+	    Result.reserve(Left.size() + Right.size());
+	    Result.insert(Result.end(), Left.begin(), Left.end());
+	    Result.insert(Result.end(), Right.begin(), Right.end());
+	    return Result;
+	}
+} // namespace
+
+const TfTokenVector&
+UEUsdGeomTetMesh::GetSchemaAttributeNames(bool bIncludeInerited)
+{
+    static TfTokenVector LocalNames = {
+		UEUsdGeomTokens->tetOrientation,
+		UEUsdGeomTokens->tetVertexIndices,
+	};
+    static TfTokenVector AllNames =
+        _ConcatenateAttributeNames(
+          UsdGeomPointBased::GetSchemaAttributeNames(true),
+          LocalNames);
+    return bIncludeInerited ? AllNames : LocalNames;
+}
+
+// ===================================================================== //
+// Feel free to add custom code below this line. It will be preserved by
+// the code generator.
+// ===================================================================== //
+// --(BEGIN CUSTOM CODE)--
+
+PXR_NAMESPACE_CLOSE_SCOPE
+
+#endif // USE_USD_SDK
diff --git a/Engine/Plugins/Experimental/ChaosCachingUSD/Source/ChaosCachingUSD/Private/ChaosCachingUSD/UEUsdGeomTetMesh.h b/Engine/Plugins/Experimental/ChaosCachingUSD/Source/ChaosCachingUSD/Private/ChaosCachingUSD/UEUsdGeomTetMesh.h
new file mode 100644
--- /dev/null
+++ b/Engine/Plugins/Experimental/ChaosCachingUSD/Source/ChaosCachingUSD/Private/ChaosCachingUSD/UEUsdGeomTetMesh.h
@@ -0,0 +1,104 @@
+// Copyright Epic Games, Inc. All Rights Reserved.
+#pragma once
+#if USE_USD_SDK
+
+#include "USDIncludesStart.h"
+	#include "pxr/base/gf/vec4i.h"
+	#include "pxr/base/tf/type.h"
+	#include "pxr/base/vt/value.h"
+
+	#include "pxr/usd/usdGeom/mesh.h"
+	#include "pxr/usd/usd/prim.h"
+	#include "pxr/usd/usd/stage.h"
+	#include "pxr/usd/usdGeom/tokens.h"
+	#include "pxr/usd/usd/timeCode.h"
+#include "USDIncludesEnd.h"
+
+PXR_NAMESPACE_OPEN_SCOPE
+
+class SdfAssetPath;
+
+/**
+ * \class UEUsdGeomTetMesh
+ *
+ * A Tetrahedral mesh prim, derived from \c UsdGeomMesh, which is used for 
+ * drawing the exterior hull of the tetrahedron mesh.  This class adds a
+ * topology array for tetrahedra consisting of \c GfVec4i, and an attribute
+ * that specifies the tetrahedron orientation/winding order.
+ * 
+ * The constructor defines and sets the subdivision scheme attribute to `none`
+ * so that by default the geometry rendered by the \c UsdGeomMesh base class
+ * is polygonal.
+ */
+class UEUsdGeomTetMesh : public UsdGeomMesh
+{
+public:
+    static const UsdSchemaKind schemaKind = UsdSchemaKind::ConcreteTyped;                                                                          
+
+    /** Construct a UEUsdGeomTetMesh from a UsdPrim \p prim. */
+    explicit UEUsdGeomTetMesh(const UsdPrim& Prim=UsdPrim());
+    explicit UEUsdGeomTetMesh(const UsdSchemaBase& SchemaObj);
+
+    virtual ~UEUsdGeomTetMesh();
+
+    static const TfTokenVector &
+    GetSchemaAttributeNames(bool bIncludeInherited=true);
+
+    static UEUsdGeomTetMesh
+    Get(const UsdStagePtr &Stage, const SdfPath &Path);
+
+    static UEUsdGeomTetMesh
+    Define(const UsdStagePtr &Stage, const SdfPath &Path);
+   
+protected:
+    UsdSchemaKind _GetSchemaKind() const override;
+
+private:
+    // Needs to invoke _GetStaticTfType.
+    friend class UsdSchemaRegistry;
+    static const TfType &_GetStaticTfType();
+
+    static bool _IsTypedSchema();
+
+    // Override SchemaBase virtuals.
+    virtual const TfType &_GetTfType() const;
+
+public:
+    /**
+	 * Tetrahedron indices. Each element of this array contains the
+     * indices of the vertices of a tetrahedron. The points are authored
+     * into the "points" attribute, just as they are for other PointBased    
+     * prims.
+     *
+     *   C++ Type: VtArray<GfVec4i>
+     *   Usd Type: SdfValueTypeNames->Int4Array
+     *   Variability: SdfVariabilityVarying
+     *   Fallback Value: No Fallback
+	 */
+    UsdAttribute GetTetVertexIndicesAttr() const;
+    
+    /**
+	 * See GetTetVertexIndicesAttr(), and also
+     * \ref Usd_Create_Or_Get_Property for when to use Get vs Create.
+     * If specified, author \p DefaultValue as the attribute's default,
+     * sparsely (when it makes sense to do so) if \p WriteSparsely is \c true -
+     * the default for \p writeSparsely is \c false.
+	 */
+    UsdAttribute CreateTetVertexIndicesAttr(VtValue const &DefaultValue = VtValue(), bool WriteSparsely=false) const;
+    
+public:
+	/**
+	 * Tetrahedron orientation. Given [A,B,C,D] if your right hand 
+	 * is at A, if your fingers wrap to A-B, then A-C, and your thumb 
+	 * points at D.
+	 */
+	UsdAttribute GetTetOrientationAttr() const;
+
+	/**
+	* See GetTetOrientationAttr().
+	*/
+	UsdAttribute CreateTetOrientationAttr(VtValue const& DefaultValue = VtValue(), bool WriteSparsely = false) const;
+};
+
+PXR_NAMESPACE_CLOSE_SCOPE
+#endif // USE_USD_SDK
\ No newline at end of file
diff --git a/Engine/Plugins/Experimental/ChaosCachingUSD/Source/ChaosCachingUSD/Private/ChaosCachingUSD/UEUsdGeomTokens.cpp b/Engine/Plugins/Experimental/ChaosCachingUSD/Source/ChaosCachingUSD/Private/ChaosCachingUSD/UEUsdGeomTokens.cpp
new file mode 100644
--- /dev/null
+++ b/Engine/Plugins/Experimental/ChaosCachingUSD/Source/ChaosCachingUSD/Private/ChaosCachingUSD/UEUsdGeomTokens.cpp
@@ -0,0 +1,22 @@
+// Copyright Epic Games, Inc. All Rights Reserved.
+
+#include "ChaosCachingUSD/UEUsdGeomTokens.h"
+
+#if USE_USD_SDK
+
+PXR_NAMESPACE_OPEN_SCOPE
+
+UEUsdGeomTokensType::UEUsdGeomTokensType()
+    : tetOrientation("tetOrientation", TfToken::Immortal)
+    , tetVertexIndices("tetVertexIndices", TfToken::Immortal)
+    , allTokens({
+            tetOrientation,
+            tetVertexIndices
+        })
+{}
+
+TfStaticData<UEUsdGeomTokensType> UEUsdGeomTokens;
+
+PXR_NAMESPACE_CLOSE_SCOPE
+
+#endif // USE_USD_SDK
diff --git a/Engine/Plugins/Experimental/ChaosCachingUSD/Source/ChaosCachingUSD/Private/ChaosCachingUSD/UEUsdGeomTokens.h b/Engine/Plugins/Experimental/ChaosCachingUSD/Source/ChaosCachingUSD/Private/ChaosCachingUSD/UEUsdGeomTokens.h
new file mode 100644
--- /dev/null
+++ b/Engine/Plugins/Experimental/ChaosCachingUSD/Source/ChaosCachingUSD/Private/ChaosCachingUSD/UEUsdGeomTokens.h
@@ -0,0 +1,38 @@
+// Copyright Epic Games, Inc. All Rights Reserved.
+
+#pragma once
+#if USE_USD_SDK
+
+#include "USDIncludesStart.h"
+	#include "pxr/base/tf/staticData.h"
+	#include "pxr/base/tf/token.h"
+#include "USDIncludesEnd.h"
+
+#include <vector>
+
+PXR_NAMESPACE_OPEN_SCOPE
+
+struct UEUsdGeomTokensType 
+{
+	CHAOSCACHINGUSD_API UEUsdGeomTokensType();
+
+    /// \brief "tetOrientation"
+    ///
+    /// UEUsdGeomTetMesh
+    const TfToken tetOrientation;
+    /// \brief "tetVetexIndices"
+    ///
+    /// UEUsdGeomTetMesh
+    const TfToken tetVertexIndices;
+    /// A vector of all of the tokens listed above.
+    const std::vector<TfToken> allTokens;
+};
+
+/// \var UEUsdGeomTokens
+///
+/// A global variable for access to static tokens.
+extern CHAOSCACHINGUSD_API TfStaticData<UEUsdGeomTokensType> UEUsdGeomTokens;
+
+PXR_NAMESPACE_CLOSE_SCOPE
+
+#endif // USE_USD_SDK
diff --git a/Engine/Plugins/Experimental/ChaosCachingUSD/Source/ChaosCachingUSD/Private/ChaosCachingUSD/Util.cpp b/Engine/Plugins/Experimental/ChaosCachingUSD/Source/ChaosCachingUSD/Private/ChaosCachingUSD/Util.cpp
new file mode 100644
--- /dev/null
+++ b/Engine/Plugins/Experimental/ChaosCachingUSD/Source/ChaosCachingUSD/Private/ChaosCachingUSD/Util.cpp
@@ -0,0 +1,100 @@
+// Copyright Epic Games, Inc. All Rights Reserved.
+
+#include "ChaosCachingUSD/Util.h"
+
+#if USE_USD_SDK
+
+#include "USDIncludesStart.h"
+	#include "pxr/usd/usdGeom/pointBased.h"
+	#include "pxr/usd/usdGeom/xform.h"
+#include "USDIncludesEnd.h"
+
+void 
+UE::ChaosCachingUSDUtil::Private::SaveStage(
+	const pxr::UsdStageRefPtr& Stage,
+	const pxr::UsdTimeCode& FirstFrame,
+	const pxr::UsdTimeCode& LastFrame)
+{
+	FScopedUsdAllocs UEAllocs; // Use USD memory allocator
+	if (FirstFrame.IsNumeric() && LastFrame.IsNumeric())
+	{
+		// Only expand the existing frame range.
+		const double StageStart = Stage->GetStartTimeCode();
+		const double StageEnd = Stage->GetEndTimeCode();
+		const bool HasRange = Stage->HasAuthoredTimeCodeRange();
+		if (!HasRange || FirstFrame.GetValue() < StageStart)
+		{
+			Stage->SetStartTimeCode(FirstFrame.GetValue());
+		}
+		if (!HasRange || LastFrame.GetValue() > StageEnd)
+		{
+			Stage->SetEndTimeCode(LastFrame.GetValue());
+		}
+	}
+	Stage->GetRootLayer()->Save();
+}
+
+void 
+UE::ChaosCachingUSDUtil::Private::DefineAncestorTransforms(
+	pxr::UsdStageRefPtr& Stage,
+	const pxr::SdfPath& PrimPath)
+{
+	FScopedUsdAllocs UEAllocs; // Use USD memory allocator
+	pxr::SdfPath Path(PrimPath.GetParentPath());
+	//while (Path != pxr::SdfPath::EmptyPath())
+	while (!Path.IsAbsoluteRootPath())
+	{
+		pxr::UsdGeomXform::Define(Stage, Path);
+		Path = Path.GetParentPath();
+	}
+}
+
+void
+UE::ChaosCachingUSDUtil::Private::ComputeExtent(
+	const pxr::VtArray<pxr::GfVec3f>& Points,
+	pxr::VtVec3fArray& Extent)
+{
+	FScopedUsdAllocs UEAllocs; // Use USD memory allocator
+	pxr::GfRange3f Bounds;
+	for (size_t i = 0; i < Points.size(); i++)
+	{
+		Bounds = Bounds.UnionWith(Points[i]);
+	}
+	Extent[0] = Bounds.GetMin();
+	Extent[1] = Bounds.GetMax();
+}
+
+bool 
+UE::ChaosCachingUSDUtil::Private::SetPointsExtentAndPrimXform(
+	pxr::UsdPrim& Prim, 
+	const pxr::GfMatrix4d& PrimXf, 
+	const pxr::VtArray<pxr::GfVec3f>& Points, 
+	const pxr::VtArray<pxr::GfVec3f>& Vels, 
+	const pxr::UsdTimeCode& Time)
+{
+	FScopedUsdAllocs UEAllocs; // Use USD memory allocator
+
+	pxr::VtVec3fArray Extent(2);
+	UE::ChaosCachingUSDUtil::Private::ComputeExtent(Points, Extent);
+
+	pxr::UsdGeomPointBased PointBased(Prim);
+	PointBased.MakeMatrixXform().Set(PrimXf, Time);
+	PointBased.CreateExtentAttr().Set(Extent, Time);
+	if (!PointBased.CreatePointsAttr().Set(Points, Time)) return false;
+	if (!Vels.empty() && !PointBased.CreateVelocitiesAttr().Set(Vels, Time)) return false;
+	return true;
+}
+
+bool 
+UE::ChaosCachingUSDUtil::Private::SetPointsExtentAndPrimXform(
+	pxr::UsdPrim& Prim,
+	const pxr::GfMatrix4d& PrimXf,
+	const pxr::VtArray<pxr::GfVec3f>& Points,
+	const pxr::UsdTimeCode& Time)
+{
+	FScopedUsdAllocs UEAllocs; // Use USD memory allocator
+	pxr::VtArray<pxr::GfVec3f> Vels;
+	return UE::ChaosCachingUSDUtil::Private::SetPointsExtentAndPrimXform(Prim, PrimXf, Points, Vels, Time);
+}
+
+#endif // USE_USD_SDK
diff --git a/Engine/Plugins/Experimental/ChaosCachingUSD/Source/ChaosCachingUSD/Private/ChaosCachingUSD/Util.h b/Engine/Plugins/Experimental/ChaosCachingUSD/Source/ChaosCachingUSD/Private/ChaosCachingUSD/Util.h
new file mode 100644
--- /dev/null
+++ b/Engine/Plugins/Experimental/ChaosCachingUSD/Source/ChaosCachingUSD/Private/ChaosCachingUSD/Util.h
@@ -0,0 +1,101 @@
+// Copyright Epic Games, Inc. All Rights Reserved.
+#pragma once
+#if USE_USD_SDK
+
+#include "USDMemory.h"
+#include "USDIncludesStart.h"
+	#include "pxr/usd/usd/usdaFileFormat.h"
+	#include "pxr/usd/usd/usdcFileFormat.h"
+	#include "pxr/usd/usd/usdFileFormat.h"
+
+	#include "pxr/usd/usd/tokens.h"
+	#include "pxr/usd/usdUtils/stitchClips.h"
+
+	#include "pxr/base/tf/fileUtils.h"
+	#include "pxr/base/tf/pathUtils.h"
+
+	#include "pxr/base/vt/array.h"
+	#include "pxr/base/vt/dictionary.h"
+	#include "pxr/base/vt/value.h"
+#include "USDIncludesEnd.h"
+
+namespace UE::ChaosCachingUSDUtil::Private {
+
+	/**
+	 * Convert from \c pxr::VtArray<T1> to \c TArray<T2>, where \p T1 and \p T2 are binary compatible.
+	 *
+	 * Note that \c VtArray is shared copy-on-write, which makes them trivially copied. This function
+	 * copies memory, which should be avoided.
+	 */
+	template <class T1, class T2>
+	void CastConvert(const pxr::VtArray<T1>& VtArray, TArray<T2>& UEArray)
+	{
+		check(sizeof(T1) == sizeof(T2)); // reinterpret_cast works?
+		check(VtArray.size() <= TNumericLimits<int32>::Max()); // int32 index works?
+		FScopedUnrealAllocs UEAllocs; // Use UE memory allocator
+		UEArray.SetNum(VtArray.size());
+		int32 i = 0;
+		for (typename pxr::VtArray<T1>::const_iterator it = VtArray.cbegin(), itEnd = VtArray.cend();
+			it != itEnd; ++it)
+		{
+			UEArray[i++] = *reinterpret_cast<const T2*>(&*it);
+		}
+	}
+
+	/**
+	 * Linearly interpolate values in \p VtArray1 and \p VtArray2 by interpolant \p Alpha ([0,1]), 
+	 * converting from \p T1 to \p T2, where \p T1 and \p T2 are binary compatible.
+	 *
+	 * Note that \c VtArray is shared copy-on-write, which makes them trivially copied. This function
+	 * copies memory, which should be avoided.
+	 */
+	template <class T1, class T2>
+	void CastLerp(const pxr::VtArray<T1>& VtArray1, const pxr::VtArray<T1>& VtArray2, TArray<T2>& UEArray, const float Alpha)
+	{
+		check(sizeof(T1) == sizeof(T2)); // reinterpret_cast works?
+		check(VtArray1.size() <= TNumericLimits<int32>::Max()); // int32 index works?
+		check(VtArray1.size() == VtArray2.size());
+		FScopedUnrealAllocs UEAllocs; // Use UE memory allocator
+		UEArray.SetNum(VtArray1.size());
+		int32 i = 0;
+		if (Alpha <= TNumericLimits<float>::Min())
+		{
+			for (typename pxr::VtArray<T1>::const_iterator
+				it1 = VtArray1.cbegin(), itEnd1 = VtArray1.end();
+				it1 != itEnd1; ++it1)
+			{
+				UEArray[i++] = *reinterpret_cast<const T2*>(&*it1);
+			}
+		}
+		else if (Alpha >= 1.0 - TNumericLimits<float>::Min())
+		{
+			for (typename pxr::VtArray<T1>::const_iterator
+				it2 = VtArray2.cbegin(), itEnd2 = VtArray2.end();
+				it2 != itEnd2; ++it2)
+			{
+				UEArray[i++] = *reinterpret_cast<const T2*>(&*it2);
+			}
+		}
+		else
+		{
+			for (typename pxr::VtArray<T1>::const_iterator
+				it1 = VtArray1.cbegin(), itEnd1 = VtArray1.end(),
+				it2 = VtArray2.cbegin(), itEnd2 = VtArray2.cend();
+				it1 != itEnd1 && it2 != itEnd2; ++it1, ++it2)
+			{
+				UEArray[i++] = 
+					(1.0 - Alpha) * *reinterpret_cast<const T2*>(&*it1) + 
+					       Alpha  * *reinterpret_cast<const T2*>(&*it2);
+			}
+		}
+	}
+
+	void SaveStage(const pxr::UsdStageRefPtr& Stage, const pxr::UsdTimeCode& FirstFrame, const pxr::UsdTimeCode &LastFrame);
+	void DefineAncestorTransforms(pxr::UsdStageRefPtr& Stage, const pxr::SdfPath& PrimPath);
+	void ComputeExtent(const pxr::VtArray<pxr::GfVec3f>& Points, pxr::VtVec3fArray& Extent);
+	bool SetPointsExtentAndPrimXform(pxr::UsdPrim& Prim, const pxr::GfMatrix4d& PrimXf, const pxr::VtArray<pxr::GfVec3f>& Points, const pxr::VtArray<pxr::GfVec3f>& Vels, const pxr::UsdTimeCode& Time = pxr::UsdTimeCode::Default());
+	bool SetPointsExtentAndPrimXform(pxr::UsdPrim& Prim, const pxr::GfMatrix4d& PrimXf, const pxr::VtArray<pxr::GfVec3f>& Points, const pxr::UsdTimeCode& Time = pxr::UsdTimeCode::Default());
+
+} // namespace UE::ChaosCachingUSDUtil::Private
+
+#endif // USE_USD_SDK
diff --git a/Engine/Plugins/Experimental/ChaosCachingUSD/Source/ChaosCachingUSD/Public/ChaosCachingUSD/Operations.h b/Engine/Plugins/Experimental/ChaosCachingUSD/Source/ChaosCachingUSD/Public/ChaosCachingUSD/Operations.h
new file mode 100644
--- /dev/null
+++ b/Engine/Plugins/Experimental/ChaosCachingUSD/Source/ChaosCachingUSD/Public/ChaosCachingUSD/Operations.h
@@ -0,0 +1,153 @@
+// Copyright Epic Games, Inc. All Rights Reserved.
+
+#pragma once
+#if USE_USD_SDK
+
+#include "GeometryCollection/GeometryCollection.h"
+#include "GeometryCollection/ManagedArray.h"
+#include "Math/UnrealMathUtility.h"
+#include "UsdWrappers/UsdStage.h"
+
+#include "USDIncludesStart.h"
+	#include "pxr/base/gf/vec3f.h"
+	#include "pxr/base/vt/array.h"
+#include "USDIncludesEnd.h"
+
+namespace UE::ChaosCachingUSD
+{
+	template <class TV>
+	pxr::VtArray<pxr::GfVec3f> ToVtVec3Array(const TArray<TV>& V3Array)
+	{
+		FScopedUsdAllocs UsdAllocs; // Use USD memory allocator
+		pxr::VtArray<pxr::GfVec3f> RetVal(static_cast<size_t>(V3Array.Num()));
+		for (int32 i = 0; i < V3Array.Num(); i++) RetVal[i].Set(V3Array[i][0], V3Array[i][1], V3Array[i][2]);
+		return RetVal;
+	}
+
+	bool ValuesDiffer(const pxr::VtArray<pxr::GfVec3f>& A, const pxr::VtArray<pxr::GfVec3f>& B, const float Tolerance = 1.0e-8, const uint64 stride = 1)
+	{
+		if (A.size() != B.size()) return true;
+		for (size_t i = 0; i < A.size(); i += stride)
+		{
+			for (size_t j = 0; j < 3; j++)
+			{
+				if (!FMath::IsNearlyEqual(A[i][j], B[i][j], Tolerance))
+					return true;
+			}
+		}
+		return false;
+	}
+
+	/** Create a new stage, resident in the USD stage cache. */
+	CHAOSCACHINGUSD_API bool NewStage(const FString& StageName, UE::FUsdStage& Stage);
+	/** Opens an existing stage, resident in the USD stage cache. */
+	CHAOSCACHINGUSD_API bool OpenStage(const FString& StageName, UE::FUsdStage& Stage);
+
+	/** Save the stage and update the frame range (if not -DBL_MAX). */
+	CHAOSCACHINGUSD_API bool SaveStage(UE::FUsdStage& Stage, const double FirstFrame, const double LastFrame);
+
+	/** Close and remove a stage from the USD stage cache. */
+	CHAOSCACHINGUSD_API bool CloseStage(const UE::FUsdStage& Stage);
+	CHAOSCACHINGUSD_API bool CloseStage(const FString& StageName);
+
+	/**
+	 * "Value Clips" is USD's concept of splitting data up into multiple 
+	 * files, the it looks like this:
+	 *
+	 *		path/to/file.usd - top level root or "parent" file that contains references to other files
+	 *		path/to/file.topology.usd - unvarying data goes here
+	 *		path/to/file.manifest.usd - we don't use this file
+	 *		path/to/file.#.usd - data for frame # goes here, where # is some number.
+	 *
+	 * Writing value clips scenes requires juggling multiple stages. Reading value clips scenes
+	 * should be somewhat simpler as USD's composition should pick up file references.
+	 */
+
+	/** Given 'path/to/file.usd', yields \p TopologyName 'path/to/file.topology.usd' and \p TimeVaryingTemplate 'path/to/file.###.###.usd'. */
+	CHAOSCACHINGUSD_API void
+	GenerateValueClipStageNames(const FString& ParentName, FString& TopologyName, FString& TimeVaryingTemplate);
+	/** Given 'path/to/file.###.###.usd' and 1.23, yields 'path/to/file.001.230.usd'. */
+	CHAOSCACHINGUSD_API FString
+	GenerateValueClipTimeVaryingStageName(const FString& TimeVaryingTemplate, const double Time);
+
+	/** Create new parent and topology stages for value clips. */
+	CHAOSCACHINGUSD_API bool
+	NewValueClipsStages(
+		const FString& ParentStageName,
+		const FString& TopologyStageName,
+		UE::FUsdStage& ParentStage,
+		UE::FUsdStage& TopologyStage);
+	/** Create new value clips frame stage. Be sure to close these stages. */
+	CHAOSCACHINGUSD_API bool
+	NewValueClipsFrameStage(
+		const FString& TimeVaryingStageTemplate,
+		const double Time,
+		FString& FrameStageName,
+		UE::FUsdStage& FrameStage);
+	/** 
+	 * Init layers and metadata for value clips. 
+	 * 
+	 * \p ParentStage - the stage that serves as the hub and references others.
+	 * \p TopologyStage - stage that will hod unvarying scene data.
+	 * \p ParentStageName
+	 * \p TopologyStageName
+	 * \p TimeVaryingStageTemplate - template pattern for per-frame USD data files.
+	 * \p PrimPaths - the prim(s) to author value clips metadata on.  It will apply to all prims 
+	 *    below this point in the hierarchy.  Authoring the root scope is not allowed by USD.
+	 * \p StartTime - the first time sample that should be available at this point in time.
+	 * \p EndTime - the last time sample that should be available at this point in time.
+	 * \p Stride - the size of the steps taken between \p StartTime and \p EndTime.
+	 */
+	CHAOSCACHINGUSD_API bool
+	InitValueClipsTemplate(
+		UE::FUsdStage& ParentStage,
+		UE::FUsdStage& TopologyStage,
+		const FString& ParentStageName,
+		const FString& TopologyStageName,
+		const FString& TimeVaryingStageTemplate,
+		const TArray<FString>& PrimPaths,
+		const double StartTime,
+		const double EndTime,
+		const double Stride);
+	/** Helper function that opens/retreives parent and topology stages from cache. */
+	CHAOSCACHINGUSD_API bool
+	InitValueClipsTemplate(
+		const FString& ParentStageName,
+		const FString& TopologyStageName,
+		const FString& TimeVaryingStageTemplate,
+		const TArray<FString>& PrimPaths,
+		const double StartTime,
+		const double EndTime,
+		const double Stride);
+
+	/**
+	 * Define \c UEUsdGeomTetMesh at \c PrimPath with ancenstor transforms, and authors points at
+	 * usd "default" time from tet mesh \p StructureIndex (or all if \c INDEX_NONE) from \c Collection.
+	 */
+	CHAOSCACHINGUSD_API bool WriteTetMesh(UE::FUsdStage& Stage, const FString& PrimPath, const FManagedArrayCollection& Collection, const int32 StructureIndex = INDEX_NONE);
+
+	/**
+	 * Define \c UsdGeomPointBased at \c PrimPath with ancestor transforms, and author points at 
+	 * \p Time (or usd "default" time if \c -DBL_MAX) from \p Collection vertices from geometry 
+	 * \p StructureIndex (or all if \c StructureIndex == \c INDEX_NONE).
+	 */
+	CHAOSCACHINGUSD_API bool WritePoints(UE::FUsdStage& Stage, const FString& PrimPath, const double Time, const FManagedArrayCollection& Collection, const int32 StructureIndex = INDEX_NONE);
+	/** Write points and velocities to \c UsdGeomPointBased at \c PrimPath. */
+	CHAOSCACHINGUSD_API bool WritePoints(UE::FUsdStage& Stage, const FString& PrimPath, const double Time, pxr::VtArray<pxr::GfVec3f>& VtPoints, pxr::VtArray<pxr::GfVec3f>& VtVels);
+
+	/** Get time samples for an attribute. */
+	CHAOSCACHINGUSD_API bool ReadTimeSamples(const UE::FUsdStage& Stage, const FString& PrimPath, const FString& AttrName, TArray<double>& TimeSamples);
+	/** Get time samples for the points attribute. */
+	CHAOSCACHINGUSD_API bool ReadTimeSamples(const UE::FUsdStage& Stage, const FString& PrimPath, TArray<double>& TimeSamples);
+
+	CHAOSCACHINGUSD_API FString GetPointsAttrName();
+	CHAOSCACHINGUSD_API FString GetVelocityAttrName();
+	CHAOSCACHINGUSD_API bool GetBracketingTimeSamples(const UE::FUsdStage& Stage, const FString& PrimPath, const FString& AttrName, const double TargetTime, double* Lower, double* Upper);
+
+	/** Get points from an attribute. Default time is used if \p Time is \c -DBL_MAX. */
+	CHAOSCACHINGUSD_API bool ReadPoints(const UE::FUsdStage& Stage, const FString& PrimPath, const FString& AttrPath, const double Time, pxr::VtArray<pxr::GfVec3f>& Points);
+	/** Get points from the points attribute. Default time is used if \p Time is \c -DBL_MAX. */
+	CHAOSCACHINGUSD_API bool ReadPoints(const UE::FUsdStage& Stage, const FString& PrimPath, const double Time, pxr::VtArray<pxr::GfVec3f>& Points, pxr::VtArray<pxr::GfVec3f>& VtVels);
+
+} // UE::ChaosCachingUSD
+#endif // USE_USD_SDK
diff --git a/Engine/Plugins/Experimental/ChaosFlesh/ChaosFlesh.uplugin b/Engine/Plugins/Experimental/ChaosFlesh/ChaosFlesh.uplugin
--- a/Engine/Plugins/Experimental/ChaosFlesh/ChaosFlesh.uplugin
+++ b/Engine/Plugins/Experimental/ChaosFlesh/ChaosFlesh.uplugin
@@ -50,6 +50,13 @@
 			"Name": "ChaosCaching",
 			"Enabled": true
 		},
+        {
+            "Name": "ChaosCachingUSD",
+			"Enabled": true,
+			"SupportedTargetPlatforms": [
+				"Win64"
+			]
+        },
 		{
 			"Name": "GeometryProcessing",
 			"Enabled": true
@@ -61,13 +68,6 @@
 		{
 			"Name": "DeformerGraph",
 			"Enabled": true
-		},
-        {
-            "Name": "USDImporter",
-			"Enabled": true,
-			"SupportedTargetPlatforms": [
-				"Win64"
-			]
-        }
+		}
 	]
 }
diff --git a/Engine/Plugins/Experimental/ChaosFlesh/Source/ChaosFleshEngine/ChaosFleshEngine.Build.cs b/Engine/Plugins/Experimental/ChaosFlesh/Source/ChaosFleshEngine/ChaosFleshEngine.Build.cs
--- a/Engine/Plugins/Experimental/ChaosFlesh/Source/ChaosFleshEngine/ChaosFleshEngine.Build.cs
+++ b/Engine/Plugins/Experimental/ChaosFlesh/Source/ChaosFleshEngine/ChaosFleshEngine.Build.cs
@@ -38,6 +38,7 @@ namespace UnrealBuildTool.Rules
 					"CoreUObject",
 					"Chaos",
 					"ChaosCaching",
+					"ChaosCachingUSD",
 					"ChaosFlesh",
 					"DataflowCore",
 					"DataflowEngine",
@@ -49,7 +50,6 @@ namespace UnrealBuildTool.Rules
 					"RenderCore",
                     "RHI",
 					"Renderer",
-					"SimCacheUSD",
 					"UnrealUSDWrapper",
 					"USDClasses",
 					"USDUtilities",
diff --git a/Engine/Plugins/Experimental/ChaosFlesh/Source/ChaosFleshEngine/Private/ChaosCache/FleshComponentCacheAdapter.cpp b/Engine/Plugins/Experimental/ChaosFlesh/Source/ChaosFleshEngine/Private/ChaosCache/FleshComponentCacheAdapter.cpp
--- a/Engine/Plugins/Experimental/ChaosFlesh/Source/ChaosFleshEngine/Private/ChaosCache/FleshComponentCacheAdapter.cpp
+++ b/Engine/Plugins/Experimental/ChaosFlesh/Source/ChaosFleshEngine/Private/ChaosCache/FleshComponentCacheAdapter.cpp
@@ -14,7 +14,7 @@
 #include "HAL/PlatformFileManager.h"
 #include "Misc/Paths.h"
 
-#include "SimCacheUSDOperations.h"
+#include "ChaosCachingUSD/Operations.h"
 #include "USDConversionUtils.h" // for GetPrimPathForObject()
 
 #endif // USE_USD_SDK && DO_USD_CACHING
@@ -41,7 +41,7 @@ namespace Chaos
 #if USE_USD_SDK && DO_USD_CACHING
 		if (MonolithStage)
 		{
-			UE::SimCacheUSD::CloseStage(MonolithStage);
+			UE::ChaosCachingUSD::CloseStage(MonolithStage);
 		}
 #endif // USE_USD_SDK && DO_USD_CACHING
 	}
@@ -91,26 +91,26 @@ namespace Chaos
 					MinTime = FMath::Min(InTime, MinTime);
 					MaxTime = FMath::Max(InTime, MaxTime);
 
-					pxr::VtArray<pxr::GfVec3f> Points = UE::SimCacheUSD::ToVtVec3Array(Particles.XArray());
-					pxr::VtArray<pxr::GfVec3f> Vels   = UE::SimCacheUSD::ToVtVec3Array(Particles.GetV());
+					pxr::VtArray<pxr::GfVec3f> Points = UE::ChaosCachingUSD::ToVtVec3Array(Particles.XArray());
+					pxr::VtArray<pxr::GfVec3f> Vels   = UE::ChaosCachingUSD::ToVtVec3Array(Particles.GetV());
 					
 					// Try to not fill the disk with redundant data.  Finalize() does last save 
 					// to set the final frame range.
-					if (UE::SimCacheUSD::ValuesDiffer(Points, PrevPoints, 0.5) ||
-						UE::SimCacheUSD::ValuesDiffer(Vels, PrevVels, 0.5))
+					if (UE::ChaosCachingUSD::ValuesDiffer(Points, PrevPoints, 0.5) ||
+						UE::ChaosCachingUSD::ValuesDiffer(Vels, PrevVels, 0.5))
 					{
 						PrevPoints = Points;
 						PrevVels = Vels;
 
 						if (MonolithStage)
 						{
-							if (!UE::SimCacheUSD::WritePoints(MonolithStage, PrimPath, InTime, Points, Vels))
+							if (!UE::ChaosCachingUSD::WritePoints(MonolithStage, PrimPath, InTime, Points, Vels))
 							{
 								// TODO: Warn
 								return;
 							}
 							// Save every n frames?
-							UE::SimCacheUSD::SaveStage(MonolithStage, MinTime, MaxTime);
+							UE::ChaosCachingUSD::SaveStage(MonolithStage, MinTime, MaxTime);
 						}
 					}
 #else // USE_USD_SDK && DO_USD_CACHING
@@ -177,17 +177,17 @@ namespace Chaos
 					double Next = -TNumericLimits<double>::Max();
 					double PrevV = -TNumericLimits<double>::Max();
 					double NextV = -TNumericLimits<double>::Max();
-					if (!UE::SimCacheUSD::GetBracketingTimeSamples(
-							MonolithStage, PrimPath, UE::SimCacheUSD::GetPointsAttrName(), TargetTime, &Prev, &Next) ||
-						!UE::SimCacheUSD::GetBracketingTimeSamples(
-							MonolithStage, PrimPath, UE::SimCacheUSD::GetVelocityAttrName(), TargetTime, &PrevV, &NextV) ||
+					if (!UE::ChaosCachingUSD::GetBracketingTimeSamples(
+							MonolithStage, PrimPath, UE::ChaosCachingUSD::GetPointsAttrName(), TargetTime, &Prev, &Next) ||
+						!UE::ChaosCachingUSD::GetBracketingTimeSamples(
+							MonolithStage, PrimPath, UE::ChaosCachingUSD::GetVelocityAttrName(), TargetTime, &PrevV, &NextV) ||
 						Prev != PrevV ||
 						Next != NextV)
 					{
 						return;
 					}
 
-					if (!UE::SimCacheUSD::ReadPoints(MonolithStage, PrimPath, Prev, Points0, Vels0) ||
+					if (!UE::ChaosCachingUSD::ReadPoints(MonolithStage, PrimPath, Prev, Points0, Vels0) ||
 						Points0.size() != Vels0.size())
 					{
 						return;
@@ -221,7 +221,7 @@ namespace Chaos
 						return;
 					}
 
-					if (!UE::SimCacheUSD::ReadPoints(MonolithStage, PrimPath, Next, Points1, Vels1) ||
+					if (!UE::ChaosCachingUSD::ReadPoints(MonolithStage, PrimPath, Next, Points1, Vels1) ||
 						Points1.size() != Vels1.size() ||
 						Points0.size() != Points1.size())
 					{
@@ -368,7 +368,7 @@ namespace Chaos
 					pxr::VtArray<pxr::GfVec3f> Points;
 					pxr::VtArray<pxr::GfVec3f> Vels;
 
-					if (!UE::SimCacheUSD::ReadPoints(MonolithStage, PrimPath, -TNumericLimits<double>::Max(), Points, Vels))
+					if (!UE::ChaosCachingUSD::ReadPoints(MonolithStage, PrimPath, -TNumericLimits<double>::Max(), Points, Vels))
 					{
 						return;
 					}
@@ -560,13 +560,13 @@ namespace Chaos
 
 					if (MonolithStage)
 					{
-						UE::SimCacheUSD::CloseStage(MonolithStage);
+						UE::ChaosCachingUSD::CloseStage(MonolithStage);
 					}
-					if (!UE::SimCacheUSD::NewStage(FilePath, MonolithStage))
+					if (!UE::ChaosCachingUSD::NewStage(FilePath, MonolithStage))
 					{
 						return false;
 					}
-					if (!UE::SimCacheUSD::WriteTetMesh(MonolithStage, PrimPath, *RestCollection))
+					if (!UE::ChaosCachingUSD::WriteTetMesh(MonolithStage, PrimPath, *RestCollection))
 					{
 						return false;
 					}
@@ -612,9 +612,9 @@ namespace Chaos
 					{
 						if (MonolithStage)
 						{
-							UE::SimCacheUSD::CloseStage(MonolithStage);
+							UE::ChaosCachingUSD::CloseStage(MonolithStage);
 						}
-						if (!UE::SimCacheUSD::OpenStage(FilePath, MonolithStage))
+						if (!UE::ChaosCachingUSD::OpenStage(FilePath, MonolithStage))
 						{
 							return false;
 						}
@@ -643,9 +643,9 @@ namespace Chaos
 		{
 			if (!bReadOnly)
 			{
-				UE::SimCacheUSD::SaveStage(MonolithStage, MinTime, MaxTime);
+				UE::ChaosCachingUSD::SaveStage(MonolithStage, MinTime, MaxTime);
 			}
-			UE::SimCacheUSD::CloseStage(MonolithStage);
+			UE::ChaosCachingUSD::CloseStage(MonolithStage);
 			MonolithStage = UE::FUsdStage();
 		}
 
diff --git a/Engine/Plugins/Importers/USDImporter/Source/SimCacheUSD/Private/SimCacheUSDModule.cpp b/Engine/Plugins/Importers/USDImporter/Source/SimCacheUSD/Private/SimCacheUSDModule.cpp
deleted file mode 100644
diff --git a/Engine/Plugins/Importers/USDImporter/Source/SimCacheUSD/Private/SimCacheUSDOperations.cpp b/Engine/Plugins/Importers/USDImporter/Source/SimCacheUSD/Private/SimCacheUSDOperations.cpp
deleted file mode 100644
diff --git a/Engine/Plugins/Importers/USDImporter/Source/SimCacheUSD/Private/SimCacheUSDUtil.cpp b/Engine/Plugins/Importers/USDImporter/Source/SimCacheUSD/Private/SimCacheUSDUtil.cpp
deleted file mode 100644
diff --git a/Engine/Plugins/Importers/USDImporter/Source/SimCacheUSD/Private/SimCacheUSDUtil.h b/Engine/Plugins/Importers/USDImporter/Source/SimCacheUSD/Private/SimCacheUSDUtil.h
deleted file mode 100644
diff --git a/Engine/Plugins/Importers/USDImporter/Source/SimCacheUSD/Private/UEUsdGeomTetMesh.cpp b/Engine/Plugins/Importers/USDImporter/Source/SimCacheUSD/Private/UEUsdGeomTetMesh.cpp
deleted file mode 100644
diff --git a/Engine/Plugins/Importers/USDImporter/Source/SimCacheUSD/Private/UEUsdGeomTetMesh.h b/Engine/Plugins/Importers/USDImporter/Source/SimCacheUSD/Private/UEUsdGeomTetMesh.h
deleted file mode 100644
diff --git a/Engine/Plugins/Importers/USDImporter/Source/SimCacheUSD/Private/UEUsdGeomTokens.cpp b/Engine/Plugins/Importers/USDImporter/Source/SimCacheUSD/Private/UEUsdGeomTokens.cpp
deleted file mode 100644
diff --git a/Engine/Plugins/Importers/USDImporter/Source/SimCacheUSD/Private/UEUsdGeomTokens.h b/Engine/Plugins/Importers/USDImporter/Source/SimCacheUSD/Private/UEUsdGeomTokens.h
deleted file mode 100644
diff --git a/Engine/Plugins/Importers/USDImporter/Source/SimCacheUSD/Private/wrapUEUsdGeomTetMesh.h b/Engine/Plugins/Importers/USDImporter/Source/SimCacheUSD/Private/wrapUEUsdGeomTetMesh.h
deleted file mode 100644
diff --git a/Engine/Plugins/Importers/USDImporter/Source/SimCacheUSD/Public/SimCacheUSDOperations.h b/Engine/Plugins/Importers/USDImporter/Source/SimCacheUSD/Public/SimCacheUSDOperations.h
deleted file mode 100644
diff --git a/Engine/Plugins/Importers/USDImporter/Source/SimCacheUSD/Resources/plugInfo.json b/Engine/Plugins/Importers/USDImporter/Source/SimCacheUSD/Resources/plugInfo.json
deleted file mode 100644
diff --git a/Engine/Plugins/Importers/USDImporter/Source/SimCacheUSD/SimCacheUSD.build.cs b/Engine/Plugins/Importers/USDImporter/Source/SimCacheUSD/SimCacheUSD.build.cs
deleted file mode 100644
diff --git a/Engine/Plugins/Importers/USDImporter/USDImporter.uplugin b/Engine/Plugins/Importers/USDImporter/USDImporter.uplugin
--- a/Engine/Plugins/Importers/USDImporter/USDImporter.uplugin
+++ b/Engine/Plugins/Importers/USDImporter/USDImporter.uplugin
@@ -73,14 +73,6 @@
 			"Type": "Editor",
 			"LoadingPhase": "Default"
 		},
-		{
-			"Name": "SimCacheUSD",
-			"Type": "Editor",
-			"LoadingPhase": "Default",
-			"SupportedTargetPlatforms": [
-				"Win64"
-			]
-		},
 		{
 			"Name": "USDTests",
 			"Type": "Editor",
