Include log includes.

diff --git a/Engine/Plugins/Experimental/ChaosCachingUSD/Source/ChaosCachingUSD/Public/ChaosCachingUSD/Module.h b/Engine/Plugins/Experimental/ChaosCachingUSD/Source/ChaosCachingUSD/Public/ChaosCachingUSD/Module.h
--- a/Engine/Plugins/Experimental/ChaosCachingUSD/Source/ChaosCachingUSD/Public/ChaosCachingUSD/Module.h
+++ b/Engine/Plugins/Experimental/ChaosCachingUSD/Source/ChaosCachingUSD/Public/ChaosCachingUSD/Module.h
@@ -3,6 +3,8 @@
 #pragma once
 
 #include "Modules/ModuleInterface.h"
+#include "Logging/LogMacros.h"
+#include "Logging/LogVerbosity.h"
 
 DECLARE_LOG_CATEGORY_EXTERN(LogChaosCacheUSD, Verbose, All);
 
