Add p.Chaos.Caching.USD.SaveFrequency, which is the interval in frames dictating how often USD data is flushed to disk.  The default is 10.

Move USD memory allocation to cpp files.  Doing conversion from UE allocated data to USD in templates has proven problematic.

diff --git a/Engine/Plugins/Experimental/ChaosCachingUSD/Source/ChaosCachingUSD/Private/ChaosCachingUSD/Operations.cpp b/Engine/Plugins/Experimental/ChaosCachingUSD/Source/ChaosCachingUSD/Private/ChaosCachingUSD/Operations.cpp
--- a/Engine/Plugins/Experimental/ChaosCachingUSD/Source/ChaosCachingUSD/Private/ChaosCachingUSD/Operations.cpp
+++ b/Engine/Plugins/Experimental/ChaosCachingUSD/Source/ChaosCachingUSD/Private/ChaosCachingUSD/Operations.cpp
@@ -83,6 +83,8 @@ UE::ChaosCachingUSD::OpenStage(const FString& StageName, UE::FUsdStage& UsdStage
 bool 
 UE::ChaosCachingUSD::SaveStage(UE::FUsdStage& Stage, const double FirstFrame, const double LastFrame)
 {
+	FScopedUsdAllocs UsdAllocs; // Use USD memory allocator
+
 	if (FirstFrame != -TNumericLimits<double>::Max() && LastFrame != -TNumericLimits<double>::Max())
 	{
 		const double StageStart = Stage.GetStartTimeCode();
@@ -201,7 +203,7 @@ UE::ChaosCachingUSD::InitValueClipsTemplate(
 	const double EndTime,
 	const double Stride)
 {
-	FScopedUsdAllocs UEAllocs; // Use USD memory allocator
+	FScopedUsdAllocs UsdAllocs; // Use USD memory allocator
 
 	UE::FSdfLayer ParentLayer = ParentStage.GetRootLayer();
 	UE::FSdfLayer TopologyLayer = ParentStage.GetRootLayer();
@@ -317,7 +319,7 @@ UE::ChaosCachingUSD::WriteTetMesh(
 	const FManagedArrayCollection& Collection, 
 	const int32 StructureIndex)
 {
-	FScopedUsdAllocs UEAllocs; // Use USD memory allocator
+	FScopedUsdAllocs UsdAllocs; // Use USD memory allocator
 
 	const TManagedArray<FIntVector4>* Tetrahedron =
 		Collection.FindAttribute<FIntVector4>(
@@ -460,7 +462,7 @@ UE::ChaosCachingUSD::WriteTetMesh(
 			*PrimPath, *Stage.GetRootLayer().GetRealPath());
 	}
 
-	return UE::ChaosCachingUSD::SaveStage(Stage, -TNumericLimits<double>::Max(), -TNumericLimits<double>::Max());
+	return true;
 }
 
 bool
@@ -471,6 +473,8 @@ UE::ChaosCachingUSD::WritePoints(
 	pxr::VtArray<pxr::GfVec3f>& VtPoints,
 	pxr::VtArray<pxr::GfVec3f>& VtVels)
 {
+	FScopedUsdAllocs UsdAllocs; // Use USD memory allocator
+
 	pxr::GfMatrix4d GfMat(1.0);
 
 	UE::FSdfPath Path(*PrimPath);
@@ -507,7 +511,7 @@ UE::ChaosCachingUSD::WritePoints(
 	const FManagedArrayCollection& Collection, 
 	const int32 StructureIndex)
 {
-	FScopedUsdAllocs UEAllocs; // Use USD memory allocator
+	FScopedUsdAllocs UsdAllocs; // Use USD memory allocator
 
 	const TManagedArray<int32>* VertexStart =
 		Collection.FindAttribute<int32>(
@@ -566,6 +570,24 @@ UE::ChaosCachingUSD::WritePoints(
 }
 
 bool 
+UE::ChaosCachingUSD::WritePoints(
+	UE::FUsdStage& Stage, 
+	const FString& PrimPath, 
+	const double Time, 
+	const TArray<Chaos::TVector<float, 3>>& Points, 
+	const TArray<Chaos::TVector<float, 3>>& Vels)
+{
+	FScopedUsdAllocs UsdAllocs; // Use USD memory allocator
+
+	pxr::VtArray<pxr::GfVec3f> VtPoints(static_cast<size_t>(Points.Num()));
+	for (int32 i = 0; i < Points.Num(); i++) VtPoints[i].Set(Points[i][0], Points[i][1], Points[i][2]);
+	pxr::VtArray<pxr::GfVec3f> VtVels(static_cast<size_t>(Vels.Num()));
+	for (int32 i = 0; i < Vels.Num(); i++) VtPoints[i].Set(Vels[i][0], Vels[i][1], Vels[i][2]);
+
+	return WritePoints(Stage, PrimPath, Time, VtPoints, VtVels);
+}
+
+bool
 UE::ChaosCachingUSD::ReadTimeSamples(
 	const UE::FUsdStage& Stage, 
 	const FString& PrimPath, 
@@ -589,6 +611,28 @@ UE::ChaosCachingUSD::ReadTimeSamples(
 	return true;
 }
 
+uint64
+UE::ChaosCachingUSD::GetNumTimeSamples(
+	const UE::FUsdStage& Stage,
+	const FString& PrimPath,
+	const FString& AttrName)
+{
+	UE::FSdfPath Path(*PrimPath);
+	UE::FUsdPrim Prim = Stage.GetPrimAtPath(Path);
+	if (!Prim)
+	{
+		UE_LOG(LogUsd, Error, TEXT("No prim found at path '%s'."), *PrimPath);
+		return 0;
+	}
+	UE::FUsdAttribute Attr = Prim.GetAttribute(*AttrName);
+	if (!Attr)
+	{
+		UE_LOG(LogUsd, Error, TEXT("No attribute '%s' found on prim '%s'."), *AttrName, *PrimPath);
+		return 0;
+	}
+	return Attr.GetNumTimeSamples();
+}
+
 bool 
 UE::ChaosCachingUSD::ReadTimeSamples(
 	const UE::FUsdStage& Stage, 
@@ -621,7 +665,7 @@ UE::ChaosCachingUSD::GetBracketingTimeSamples(
 	double* Lower, 
 	double* Upper)
 {
-	FScopedUsdAllocs UEAllocs; // Use USD memory allocator
+	FScopedUsdAllocs UsdAllocs; // Use USD memory allocator
 
 	UE::FSdfPath Path(*PrimPath);
 	UE::FUsdPrim Prim = Stage.GetPrimAtPath(Path);
@@ -642,7 +686,6 @@ UE::ChaosCachingUSD::GetBracketingTimeSamples(
 		TargetTime, Lower, Upper, &bHasTimeSamples) && bHasTimeSamples;
 }
 
-
 bool
 UE::ChaosCachingUSD::ReadPoints(
 	const UE::FUsdStage& Stage,
@@ -651,7 +694,7 @@ UE::ChaosCachingUSD::ReadPoints(
 	const double Time,
 	pxr::VtArray<pxr::GfVec3f>& VtPoints)
 {
-	FScopedUsdAllocs UEAllocs; // Use USD memory allocator
+	FScopedUsdAllocs UsdAllocs; // Use USD memory allocator
 	
 	UE::FSdfPath Path(*PrimPath);
 	UE::FUsdPrim Prim = Stage.GetPrimAtPath(Path);
diff --git a/Engine/Plugins/Experimental/ChaosCachingUSD/Source/ChaosCachingUSD/Public/ChaosCachingUSD/Operations.h b/Engine/Plugins/Experimental/ChaosCachingUSD/Source/ChaosCachingUSD/Public/ChaosCachingUSD/Operations.h
--- a/Engine/Plugins/Experimental/ChaosCachingUSD/Source/ChaosCachingUSD/Public/ChaosCachingUSD/Operations.h
+++ b/Engine/Plugins/Experimental/ChaosCachingUSD/Source/ChaosCachingUSD/Public/ChaosCachingUSD/Operations.h
@@ -3,6 +3,7 @@
 #pragma once
 #if USE_USD_SDK
 
+#include "Chaos/Vector.h"
 #include "GeometryCollection/GeometryCollection.h"
 #include "GeometryCollection/ManagedArray.h"
 #include "Math/UnrealMathUtility.h"
@@ -134,11 +135,14 @@ namespace UE::ChaosCachingUSD
 	CHAOSCACHINGUSD_API bool WritePoints(UE::FUsdStage& Stage, const FString& PrimPath, const double Time, const FManagedArrayCollection& Collection, const int32 StructureIndex = INDEX_NONE);
 	/** Write points and velocities to \c UsdGeomPointBased at \c PrimPath. */
 	CHAOSCACHINGUSD_API bool WritePoints(UE::FUsdStage& Stage, const FString& PrimPath, const double Time, pxr::VtArray<pxr::GfVec3f>& VtPoints, pxr::VtArray<pxr::GfVec3f>& VtVels);
+	/** Copy \p Points and \p Vels to \c VtArray (with USD memory allocator), then write to USD stage. */
+	CHAOSCACHINGUSD_API bool WritePoints(UE::FUsdStage& Stage, const FString& PrimPath, const double Time, const TArray<Chaos::TVector<float, 3>>& Points, const TArray<Chaos::TVector<float, 3>>& Vels);
 
 	/** Get time samples for an attribute. */
 	CHAOSCACHINGUSD_API bool ReadTimeSamples(const UE::FUsdStage& Stage, const FString& PrimPath, const FString& AttrName, TArray<double>& TimeSamples);
 	/** Get time samples for the points attribute. */
 	CHAOSCACHINGUSD_API bool ReadTimeSamples(const UE::FUsdStage& Stage, const FString& PrimPath, TArray<double>& TimeSamples);
+	CHAOSCACHINGUSD_API uint64 GetNumTimeSamples(const UE::FUsdStage& Stage, const FString& PrimPath, const FString& AttrName);
 
 	CHAOSCACHINGUSD_API FString GetPointsAttrName();
 	CHAOSCACHINGUSD_API FString GetVelocityAttrName();
diff --git a/Engine/Plugins/Experimental/ChaosFlesh/Source/ChaosFleshEngine/Private/ChaosCache/FleshComponentCacheAdapter.cpp b/Engine/Plugins/Experimental/ChaosFlesh/Source/ChaosFleshEngine/Private/ChaosCache/FleshComponentCacheAdapter.cpp
--- a/Engine/Plugins/Experimental/ChaosFlesh/Source/ChaosFleshEngine/Private/ChaosCache/FleshComponentCacheAdapter.cpp
+++ b/Engine/Plugins/Experimental/ChaosFlesh/Source/ChaosFleshEngine/Private/ChaosCache/FleshComponentCacheAdapter.cpp
@@ -35,6 +35,10 @@ namespace Chaos
 		TEXT("p.Chaos.Caching.USD.NoClobber"),
 		CVarParams.bNoClobber,
 		TEXT("Rename rather than over write existing cach files. [def: true]"));
+	FAutoConsoleVariableRef CVarDeformableFleshCacheSaveFrequency(
+		TEXT("p.Chaos.Caching.USD.SaveFrequency"),
+		CVarParams.SaveFrequency,
+		TEXT("Interval in frames to flush USD data to disk. 2 saves every other frame, 1 saves every frame, 0 caches in memory until complete. [def: 10]"));
 
 #endif // USE_USD_SDK && DO_USD_CACHING
 
@@ -93,27 +97,19 @@ namespace Chaos
 					MinTime = FMath::Min(InTime, MinTime);
 					MaxTime = FMath::Max(InTime, MaxTime);
 
-					pxr::VtArray<pxr::GfVec3f> Points = UE::ChaosCachingUSD::ToVtVec3Array(Particles.XArray());
-					pxr::VtArray<pxr::GfVec3f> Vels   = UE::ChaosCachingUSD::ToVtVec3Array(Particles.GetV());
-					
-					// Try to not fill the disk with redundant data.  Finalize() does last save 
-					// to set the final frame range.
-					if (UE::ChaosCachingUSD::ValuesDiffer(Points, PrevPoints, 0.5) ||
-						UE::ChaosCachingUSD::ValuesDiffer(Vels, PrevVels, 0.5))
+					if (MonolithStage)
 					{
-						PrevPoints = Points;
-						PrevVels = Vels;
+						if (!UE::ChaosCachingUSD::WritePoints(MonolithStage, PrimPath, InTime, Particles.XArray(), Particles.GetV()))
+						{
+							UE_LOG(LogChaosFleshCache, Error,
+								TEXT("Failed to write points '%s' at time %g to file: '%s'"),
+								*PrimPath, InTime, *MonolithStage.GetRootLayer().GetDisplayName());
+							return;
+						}
 
-						if (MonolithStage)
+						uint64 NumTimeSamples = UE::ChaosCachingUSD::GetNumTimeSamples(MonolithStage, PrimPath, UE::ChaosCachingUSD::GetPointsAttrName());
+						if (CVarParams.SaveFrequency >= 1 && NumTimeSamples % CVarParams.SaveFrequency == 0)
 						{
-							if (!UE::ChaosCachingUSD::WritePoints(MonolithStage, PrimPath, InTime, Points, Vels))
-							{
-								UE_LOG(LogChaosFleshCache, Error,
-									TEXT("Failed to write points '%s' at time %g to file: '%s'"),
-									*PrimPath, InTime, *MonolithStage.GetRootLayer().GetDisplayName());
-								return;
-							}
-							// Save every n frames?
 							if (!UE::ChaosCachingUSD::SaveStage(MonolithStage, MinTime, MaxTime))
 							{
 								UE_LOG(LogChaosFleshCache, Error,
@@ -675,9 +671,6 @@ namespace Chaos
 	{
 #if USE_USD_SDK && DO_USD_CACHING
 		// Detach shared memory arrays.
-		PrevPoints.clear();
-		PrevVels.clear();
-		
 		if (MonolithStage)
 		{
 			if (!bReadOnly)
diff --git a/Engine/Plugins/Experimental/ChaosFlesh/Source/ChaosFleshEngine/Public/ChaosCache/FleshComponentCacheAdapter.h b/Engine/Plugins/Experimental/ChaosFlesh/Source/ChaosFleshEngine/Public/ChaosCache/FleshComponentCacheAdapter.h
--- a/Engine/Plugins/Experimental/ChaosFlesh/Source/ChaosFleshEngine/Public/ChaosCache/FleshComponentCacheAdapter.h
+++ b/Engine/Plugins/Experimental/ChaosFlesh/Source/ChaosFleshEngine/Public/ChaosCache/FleshComponentCacheAdapter.h
@@ -33,6 +33,7 @@ namespace Chaos
 		// USD caching options
 		bool bWriteBinary = true;
 		bool bNoClobber = true;
+		int32 SaveFrequency = 10;
 	};
 
 	/**
@@ -87,9 +88,6 @@ namespace Chaos
 		mutable double MinTime = TNumericLimits<double>::Max();
 		mutable double MaxTime = -TNumericLimits<double>::Max();
 
-		mutable pxr::VtArray<pxr::GfVec3f> PrevPoints;
-		mutable pxr::VtArray<pxr::GfVec3f> PrevVels;
-
 		mutable UE::FUsdStage MonolithStage;
 
 #else // USE_USD_SDK && DO_USD_CACHING
