Fix copy/paste error.

diff --git a/Engine/Plugins/Experimental/ChaosCachingUSD/Source/ChaosCachingUSD/Private/ChaosCachingUSD/Operations.cpp b/Engine/Plugins/Experimental/ChaosCachingUSD/Source/ChaosCachingUSD/Private/ChaosCachingUSD/Operations.cpp
--- a/Engine/Plugins/Experimental/ChaosCachingUSD/Source/ChaosCachingUSD/Private/ChaosCachingUSD/Operations.cpp
+++ b/Engine/Plugins/Experimental/ChaosCachingUSD/Source/ChaosCachingUSD/Private/ChaosCachingUSD/Operations.cpp
@@ -582,7 +582,7 @@ UE::ChaosCachingUSD::WritePoints(
 	pxr::VtArray<pxr::GfVec3f> VtPoints(static_cast<size_t>(Points.Num()));
 	for (int32 i = 0; i < Points.Num(); i++) VtPoints[i].Set(Points[i][0], Points[i][1], Points[i][2]);
 	pxr::VtArray<pxr::GfVec3f> VtVels(static_cast<size_t>(Vels.Num()));
-	for (int32 i = 0; i < Vels.Num(); i++) VtPoints[i].Set(Vels[i][0], Vels[i][1], Vels[i][2]);
+	for (int32 i = 0; i < Vels.Num(); i++) VtVels[i].Set(Vels[i][0], Vels[i][1], Vels[i][2]);
 
 	return WritePoints(Stage, PrimPath, Time, VtPoints, VtVels);
 }
