Make sure the filename supplied has a file extension/format that USD supports natively.

diff --git a/Engine/Plugins/Experimental/ChaosCachingUSD/Source/ChaosCachingUSD/Private/ChaosCachingUSD/Operations.cpp b/Engine/Plugins/Experimental/ChaosCachingUSD/Source/ChaosCachingUSD/Private/ChaosCachingUSD/Operations.cpp
--- a/Engine/Plugins/Experimental/ChaosCachingUSD/Source/ChaosCachingUSD/Private/ChaosCachingUSD/Operations.cpp
+++ b/Engine/Plugins/Experimental/ChaosCachingUSD/Source/ChaosCachingUSD/Private/ChaosCachingUSD/Operations.cpp
@@ -41,6 +41,36 @@ UE::ChaosCachingUSD::NewStage(const FString& StageName, UE::FUsdStage& UsdStage)
 		return false;
 	}
 
+	// Verify that StageName is a supported format.  Note that there's GetAllSupportedFileFormats()
+	// and GetNativeFileFormats().  In this case, I think since we're creating a stage, we want
+	// a native format.
+	FString FileFormat = FPaths::GetExtension(StageName, false); // no dot
+	TArray<FString> SupportedFormats = UnrealUSDWrapper::GetNativeFileFormats();
+	bool bFound = false;
+	for (const FString& Ext : SupportedFormats)
+	{
+		if (FileFormat == Ext)
+		{
+			bFound = true;
+			break;
+		}
+	}
+	if (!bFound)
+	{
+		FString Formats;
+		for (int32 i=0; i < SupportedFormats.Num(); i++)
+		{
+			const FString& Ext = SupportedFormats[i];
+			Formats.Append(Ext);
+			if (i < SupportedFormats.Num() - 1) Formats.AppendChar(' ');
+		}
+
+		UE_LOG(LogUsd, Error, 
+			TEXT("Failed to create new USD stage: '%s', file format '%s' is not native to USD. Supported formats: '%s'"),
+			*StageName, *FileFormat, *Formats);
+		return false;
+	}
+
 	// USD caches all stages you open/create, unless you tell it not to.
 	// For convenience, we're using the global/default stage cache.
 	UsdStage = UnrealUSDWrapper::NewStage(*StageName);
