Chaos - Fixed the Experimental status display in the Plugin browser for the Dataflow, Flesh, and CachingUSD plugins.

#tests Opened the Plugins Browser and checked the status of those plugins.
#rnx

diff --git a/Engine/Plugins/Experimental/ChaosCachingUSD/ChaosCachingUSD.uplugin b/Engine/Plugins/Experimental/ChaosCachingUSD/ChaosCachingUSD.uplugin
--- a/Engine/Plugins/Experimental/ChaosCachingUSD/ChaosCachingUSD.uplugin
+++ b/Engine/Plugins/Experimental/ChaosCachingUSD/ChaosCachingUSD.uplugin
@@ -12,7 +12,7 @@
 	"SupportURL": "",
 	"EnabledByDefault": false,
 	"CanContainContent": true,
-	"IsBetaVersion": true,
+	"IsExperimentalVersion": true,
 	"Installed": false,
 	"SupportedTargetPlatforms": [
 		"Win64"
diff --git a/Engine/Plugins/Experimental/ChaosFlesh/ChaosFlesh.uplugin b/Engine/Plugins/Experimental/ChaosFlesh/ChaosFlesh.uplugin
--- a/Engine/Plugins/Experimental/ChaosFlesh/ChaosFlesh.uplugin
+++ b/Engine/Plugins/Experimental/ChaosFlesh/ChaosFlesh.uplugin
@@ -12,7 +12,6 @@
     "SupportURL" : "",
     "EnabledByDefault" : false,
     "CanContainContent" : true,
-    "IsBetaVersion" : true,
 	"IsExperimentalVersion": true,
     "Installed" : false,
 	"Modules": [
diff --git a/Engine/Plugins/Experimental/Dataflow/Dataflow.uplugin b/Engine/Plugins/Experimental/Dataflow/Dataflow.uplugin
--- a/Engine/Plugins/Experimental/Dataflow/Dataflow.uplugin
+++ b/Engine/Plugins/Experimental/Dataflow/Dataflow.uplugin
@@ -12,7 +12,6 @@
     "SupportURL" : "",
     "EnabledByDefault" : true,
     "CanContainContent" : true,
-    "IsBetaVersion" : true,
 	"IsExperimentalVersion": true,
     "Installed" : false,
 	"Modules": [
