ChaosFlesh: fixed Chaos Cache log spamming when caching an empty flesh asset

diff --git a/Engine/Plugins/Experimental/ChaosCachingUSD/Source/ChaosCachingUSD/Private/ChaosCachingUSD/Util.cpp b/Engine/Plugins/Experimental/ChaosCachingUSD/Source/ChaosCachingUSD/Private/ChaosCachingUSD/Util.cpp
--- a/Engine/Plugins/Experimental/ChaosCachingUSD/Source/ChaosCachingUSD/Private/ChaosCachingUSD/Util.cpp
+++ b/Engine/Plugins/Experimental/ChaosCachingUSD/Source/ChaosCachingUSD/Private/ChaosCachingUSD/Util.cpp
@@ -81,7 +81,7 @@ UE::ChaosCachingUSDUtil::Private::SetPointsExtentAndPrimXform(
 	PointBased.MakeMatrixXform().Set(PrimXf, Time);
 	PointBased.CreateExtentAttr().Set(Extent, Time);
 	if (!PointBased.CreatePointsAttr().Set(Points, Time)) return false;
-	if (!Vels.empty() && !PointBased.CreateVelocitiesAttr().Set(Vels, Time)) return false;
+	if (!PointBased.CreateVelocitiesAttr().Set(Vels, Time)) return false;
 	return true;
 }
 
diff --git a/Engine/Plugins/Experimental/ChaosFlesh/Source/ChaosFleshEngine/Private/ChaosCache/FleshComponentCacheAdapter.cpp b/Engine/Plugins/Experimental/ChaosFlesh/Source/ChaosFleshEngine/Private/ChaosCache/FleshComponentCacheAdapter.cpp
--- a/Engine/Plugins/Experimental/ChaosFlesh/Source/ChaosFleshEngine/Private/ChaosCache/FleshComponentCacheAdapter.cpp
+++ b/Engine/Plugins/Experimental/ChaosFlesh/Source/ChaosFleshEngine/Private/ChaosCache/FleshComponentCacheAdapter.cpp
@@ -92,7 +92,7 @@ namespace Chaos
 				const FParticles& Particles = Evolution->Particles();
 
 				const uint32 NumParticles = Particles.Size();
-				if (NumParticles > 0)
+				//We always write a cache even if there are no particles
 				{
 #if USE_USD_SDK && DO_USD_CACHING
 					FScopedUsdAllocs UsdAllocs; // Use USD memory allocator
