ChaosCachingUSD: remove build setting forcing module to use C++17

ChaosCachingUSD inherited this setting because the USD modules
it uses previously required C++17. Now that those modules fully support
C++20, it is no longer necessary for ChaosCachingUSD to be set back
to C++17.

diff --git a/Engine/Plugins/Experimental/ChaosCachingUSD/Source/ChaosCachingUSD/ChaosCachingUSD.build.cs b/Engine/Plugins/Experimental/ChaosCachingUSD/Source/ChaosCachingUSD/ChaosCachingUSD.build.cs
--- a/Engine/Plugins/Experimental/ChaosCachingUSD/Source/ChaosCachingUSD/ChaosCachingUSD.build.cs
+++ b/Engine/Plugins/Experimental/ChaosCachingUSD/Source/ChaosCachingUSD/ChaosCachingUSD.build.cs
@@ -10,11 +10,6 @@ namespace UnrealBuildTool.Rules
 	{
 		public ChaosCachingUSD(ReadOnlyTargetRules Target) : base(Target)
 		{
-			// Does not compile with C++20:
-			// error C4002: too many arguments for function-like macro invocation 'TF_PP_CAT_IMPL'
-			// warning C5103: pasting '"TF_LOG_STACK_TRACE_ON_ERROR"' and '"TF_LOG_STACK_TRACE_ON_WARNING"' does not result in a valid preprocessing token
-			CppStandard = CppStandardVersion.Cpp17;
-
 			bUseRTTI = true;
 
 			PublicDependencyModuleNames.AddRange(
