ChaosCachingUSD:
- Fixed allocator crashes when starting the editor caused by order of the GeomTetMesh initializations.
- Deprecated unused LogChaosCacheUSD.

[RN] ChaosCachingUSD - Fixed allocator crashes when starting the editor caused by the order of the GeomTetMesh initializations.#jira none

diff --git a/Engine/Plugins/Experimental/ChaosCachingUSD/Source/ChaosCachingUSD/Private/ChaosCachingUSD/Module.cpp b/Engine/Plugins/Experimental/ChaosCachingUSD/Source/ChaosCachingUSD/Private/ChaosCachingUSD/Module.cpp
--- a/Engine/Plugins/Experimental/ChaosCachingUSD/Source/ChaosCachingUSD/Private/ChaosCachingUSD/Module.cpp
+++ b/Engine/Plugins/Experimental/ChaosCachingUSD/Source/ChaosCachingUSD/Private/ChaosCachingUSD/Module.cpp
@@ -1,17 +1,33 @@
 // Copyright Epic Games, Inc. All Rights Reserved.
 
 #include "ChaosCachingUSD/Module.h"
-
 #include "Modules/ModuleManager.h"
+
+#if USE_USD_SDK
+#include "ChaosCachingUSD/UEUsdGeomTetMesh.h"
 #include "Interfaces/IPluginManager.h"
 #include "Misc/Paths.h"
 
 #include "UnrealUSDWrapper.h"
 #include "USDMemory.h"
 
+#include "USDIncludesStart.h"
+#include "pxr/usd/sdf/assetPath.h"
+#include "pxr/usd/sdf/types.h"
+#include "pxr/usd/usd/schemaRegistry.h"
+#include "pxr/usd/usd/typed.h"
+#include "USDIncludesEnd.h"
+
 IMPLEMENT_MODULE_USD(FChaosCachingUSDModule, ChaosCachingUSD);
+#else
+IMPLEMENT_MODULE(FChaosCachingUSDModule, ChaosCachingUSD);
+#endif // USE_USD_SDK
+
 
+// UE_DEPRECATED(5.7, "Use LogUsd instead")
+PRAGMA_DISABLE_DEPRECATION_WARNINGS
 DEFINE_LOG_CATEGORY(LogChaosCacheUSD)
+PRAGMA_ENABLE_DEPRECATION_WARNINGS
 
 void FChaosCachingUSDModule::StartupModule()
 {
@@ -28,5 +44,13 @@ void FChaosCachingUSDModule::StartupModule()
 				FString(TEXT("Resources"))));
 
 	UnrealUSDWrapper::RegisterPlugins(ChaosCachingUSDResourcesDir);
+
+	// Register the tetmesh schema
+	{
+		using namespace pxr;
+		FScopedUsdAllocs UEAllocs;  // Use USD memory allocator
+		TfType::Define<UEUsdGeomTetMesh, TfType::Bases<UsdGeomMesh> >();
+		TfType::AddAlias<UsdSchemaBase, UEUsdGeomTetMesh>("UETetMesh");
+	}
 #endif // USE_USD_SDK
 }
diff --git a/Engine/Plugins/Experimental/ChaosCachingUSD/Source/ChaosCachingUSD/Private/ChaosCachingUSD/UEUsdGeomTetMesh.cpp b/Engine/Plugins/Experimental/ChaosCachingUSD/Source/ChaosCachingUSD/Private/ChaosCachingUSD/UEUsdGeomTetMesh.cpp
--- a/Engine/Plugins/Experimental/ChaosCachingUSD/Source/ChaosCachingUSD/Private/ChaosCachingUSD/UEUsdGeomTetMesh.cpp
+++ b/Engine/Plugins/Experimental/ChaosCachingUSD/Source/ChaosCachingUSD/Private/ChaosCachingUSD/UEUsdGeomTetMesh.cpp
@@ -8,43 +8,26 @@
 #include "USDMemory.h"
 #include "USDIncludesStart.h"
 	#include "pxr/usd/sdf/assetPath.h"
-	#include "pxr/usd/sdf/types.h"
-	#include "pxr/usd/usd/schemaRegistry.h"
-	#include "pxr/usd/usd/typed.h"
-
-PXR_NAMESPACE_OPEN_SCOPE
-// Register the schema with the TfType system - need to do this inside 
-// USDIncludesStart/End.h so that warnings are disabled, and re-enabled 
-// in USDIncludesEnd.h.
-#pragma warning(disable: 4191) /* 'reinterpret_cast': unsafe conversion during TF_REGISTRY_FUNCTION(TfType) */
-TF_REGISTRY_FUNCTION(TfType)
-{
-	FScopedUsdAllocs UEAllocs; // Use USD memory allocator
-	TfType::Define<UEUsdGeomTetMesh, TfType::Bases<UsdGeomMesh> >();
-	TfType::AddAlias<UsdSchemaBase, UEUsdGeomTetMesh>("UETetMesh");
-}
-PXR_NAMESPACE_CLOSE_SCOPE
-
 #include "USDIncludesEnd.h"
 
 PXR_NAMESPACE_OPEN_SCOPE
 
 UEUsdGeomTetMesh::UEUsdGeomTetMesh(const UsdPrim& Prim)
-    : UsdGeomMesh(Prim)          
+	: UsdGeomMesh(Prim)
 {
-    if(pxr::UsdAttribute SubdivisionAttr = UsdGeomMesh::CreateSubdivisionSchemeAttr())
-    {
-        SubdivisionAttr.Set(pxr::UsdGeomTokens->none);
-    }
+	if(pxr::UsdAttribute SubdivisionAttr = UsdGeomMesh::CreateSubdivisionSchemeAttr())
+	{
+		SubdivisionAttr.Set(pxr::UsdGeomTokens->none);
+	}
 }
 
 UEUsdGeomTetMesh::UEUsdGeomTetMesh(const UsdSchemaBase& SchemaObj)
-    : UsdGeomMesh(SchemaObj)          
+	: UsdGeomMesh(SchemaObj)
 {
-    if(pxr::UsdAttribute SubdivisionAttr = UsdGeomMesh::CreateSubdivisionSchemeAttr())
-    {
-        SubdivisionAttr.Set(pxr::UsdGeomTokens->none);
-    }
+	if(pxr::UsdAttribute SubdivisionAttr = UsdGeomMesh::CreateSubdivisionSchemeAttr())
+	{
+		SubdivisionAttr.Set(pxr::UsdGeomTokens->none);
+	}
 }
 
 UEUsdGeomTetMesh::~UEUsdGeomTetMesh()
@@ -64,51 +47,51 @@ UEUsdGeomTetMesh::Get(const UsdStagePtr &Stage, const SdfPath &Path)
 UEUsdGeomTetMesh
 UEUsdGeomTetMesh::Define(const UsdStagePtr &Stage, const SdfPath &Path)
 {
-    static TfToken UsdPrimTypeName("UETetMesh");
-    if(!Stage)
-    {
-        TF_CODING_ERROR("Invalid stage");
-        return UEUsdGeomTetMesh();
-    }
-    return UEUsdGeomTetMesh(Stage->DefinePrim(Path, UsdPrimTypeName));
+	static TfToken UsdPrimTypeName("UETetMesh");
+	if(!Stage)
+	{
+		TF_CODING_ERROR("Invalid stage");
+		return UEUsdGeomTetMesh();
+	}
+	return UEUsdGeomTetMesh(Stage->DefinePrim(Path, UsdPrimTypeName));
 }
 
 UsdSchemaKind
 UEUsdGeomTetMesh::_GetSchemaKind() const
 {
-    return UEUsdGeomTetMesh::schemaKind;
+	return UEUsdGeomTetMesh::schemaKind;
 }
 
 const TfType &
 UEUsdGeomTetMesh::_GetStaticTfType()
 {
-    static TfType Type = TfType::Find<UEUsdGeomTetMesh>();
-    return Type;
+	static TfType Type = TfType::Find<UEUsdGeomTetMesh>();
+	return Type;
 }
 
 bool
 UEUsdGeomTetMesh::_IsTypedSchema()
 {
-    static bool bIsTyped = _GetStaticTfType().IsA<UsdTyped>();
-    return bIsTyped;
+	static bool bIsTyped = _GetStaticTfType().IsA<UsdTyped>();
+	return bIsTyped;
 }
 
 const TfType &
 UEUsdGeomTetMesh::_GetTfType() const
 {
-    return _GetStaticTfType();
+	return _GetStaticTfType();
 }
 
 UsdAttribute
 UEUsdGeomTetMesh::GetTetVertexIndicesAttr() const
 {
-    return GetPrim().GetAttribute(UEUsdGeomTokens->tetVertexIndices);
+	return GetPrim().GetAttribute(UEUsdGeomTokens->tetVertexIndices);
 }
 
 UsdAttribute
 UEUsdGeomTetMesh::CreateTetVertexIndicesAttr(VtValue const &DefaultValue, bool bWriteSparsely) const
 {
-    return UsdSchemaBase::_CreateAttr(
+	return UsdSchemaBase::_CreateAttr(
 		UEUsdGeomTokens->tetVertexIndices,
 		SdfValueTypeNames->Int4Array,
 		false,    // custom
@@ -139,26 +122,26 @@ namespace {
 	static inline TfTokenVector
 	_ConcatenateAttributeNames(const TfTokenVector& Left, const TfTokenVector& Right)
 	{
-	    TfTokenVector Result;
-	    Result.reserve(Left.size() + Right.size());
-	    Result.insert(Result.end(), Left.begin(), Left.end());
-	    Result.insert(Result.end(), Right.begin(), Right.end());
-	    return Result;
+		TfTokenVector Result;
+		Result.reserve(Left.size() + Right.size());
+		Result.insert(Result.end(), Left.begin(), Left.end());
+		Result.insert(Result.end(), Right.begin(), Right.end());
+		return Result;
 	}
 } // namespace
 
 const TfTokenVector&
 UEUsdGeomTetMesh::GetSchemaAttributeNames(bool bIncludeInerited)
 {
-    static TfTokenVector LocalNames = {
+	static TfTokenVector LocalNames = {
 		UEUsdGeomTokens->tetOrientation,
 		UEUsdGeomTokens->tetVertexIndices,
 	};
-    static TfTokenVector AllNames =
-        _ConcatenateAttributeNames(
-          UsdGeomPointBased::GetSchemaAttributeNames(true),
-          LocalNames);
-    return bIncludeInerited ? AllNames : LocalNames;
+	static TfTokenVector AllNames =
+		_ConcatenateAttributeNames(
+		  UsdGeomPointBased::GetSchemaAttributeNames(true),
+		  LocalNames);
+	return bIncludeInerited ? AllNames : LocalNames;
 }
 
 // ===================================================================== //
diff --git a/Engine/Plugins/Experimental/ChaosCachingUSD/Source/ChaosCachingUSD/Public/ChaosCachingUSD/Module.h b/Engine/Plugins/Experimental/ChaosCachingUSD/Source/ChaosCachingUSD/Public/ChaosCachingUSD/Module.h
--- a/Engine/Plugins/Experimental/ChaosCachingUSD/Source/ChaosCachingUSD/Public/ChaosCachingUSD/Module.h
+++ b/Engine/Plugins/Experimental/ChaosCachingUSD/Source/ChaosCachingUSD/Public/ChaosCachingUSD/Module.h
@@ -6,6 +6,7 @@
 #include "Logging/LogMacros.h"
 #include "Logging/LogVerbosity.h"
 
+UE_DEPRECATED(5.7, "Use LogUsd instead")
 DECLARE_LOG_CATEGORY_EXTERN(LogChaosCacheUSD, Verbose, All);
 
 class FChaosCachingUSDModule : public IModuleInterface
